Budget allocator API & debug rendering improvements

Extended Blueprint API to allow for parameters to be set via BP (as well as code + CVars).
Added debug rendering graph of budget & time taken.
Cleaned up some defines

diff --git a/Engine/Plugins/Runtime/AnimationBudgetAllocator/Source/AnimationBudgetAllocator/Private/AnimationBudgetAllocator.cpp b/Engine/Plugins/Runtime/AnimationBudgetAllocator/Source/AnimationBudgetAllocator/Private/AnimationBudgetAllocator.cpp
--- a/Engine/Plugins/Runtime/AnimationBudgetAllocator/Source/AnimationBudgetAllocator/Private/AnimationBudgetAllocator.cpp
+++ b/Engine/Plugins/Runtime/AnimationBudgetAllocator/Source/AnimationBudgetAllocator/Private/AnimationBudgetAllocator.cpp
@@ -9,6 +9,11 @@
 #include "ProfilingDebugging/CsvProfiler.h"
 #include "SkeletalMeshComponentBudgeted.h"
 #include "DrawDebugHelpers.h"
+#include "AnimationBudgetAllocatorCVars.h"
+#include "GameFramework/PlayerController.h"
+#include "GameFramework/HUD.h"
+#include "Debug/ReporterGraph.h"
+#include "Engine/Canvas.h"
 
 DECLARE_CYCLE_STAT(TEXT("InitialTick"), STAT_AnimationBudgetAllocator_Update, STATGROUP_AnimationBudgetAllocator);
 
@@ -28,292 +33,18 @@ CSV_DEFINE_CATEGORY(AnimationBudget, true);
 
 bool FAnimationBudgetAllocator::bCachedEnabled = false;
 
-static int32 GAnimationBudgetEnabled = 0;
-
-static FAutoConsoleVariableRef CVarSkelBatch_Enabled(
-	TEXT("a.Budget.Enabled"),
-	GAnimationBudgetEnabled,
-	TEXT("Values: 0/1\n")
-	TEXT("Controls whether the skeletal mesh batching system is enabled. Should be set when there are no running skeletal meshes."),
-	ECVF_Scalability);
-
-static float GBudgetInMs = 1.0f;
-
-static FAutoConsoleVariableRef CVarSkelBatch_Budget(
-	TEXT("a.Budget.BudgetMs"),
-	GBudgetInMs,
-	TEXT("Values > 0.1, Default = 1.0\n")
-	TEXT("The time in milliseconds that we allocate for skeletal mesh work to be performed.\n")
-	TEXT("When overbudget various other CVars come into play, such as a.Budget.AlwaysTickFalloffAggression and a.Budget.InterpolationFalloffAggression."),
-	FConsoleVariableDelegate::CreateLambda([](IConsoleVariable* InVariable)
-	{
-		GBudgetInMs = FMath::Max(GBudgetInMs, 0.1f);
-	}),
-	ECVF_Scalability);
-
-static float GMinQuality = 0.0f;
-
-static FAutoConsoleVariableRef CVarSkelBatch_MinQuality(
-	TEXT("a.Budget.MinQuality"),
-	GMinQuality,
-	TEXT("Values [0.0, 1.0], Default = 0.0\n")
-	TEXT("The minimum quality metric allowed. Quality is determined simply by NumComponentsTickingThisFrame / NumComponentsThatWeNeedToTick.\n")
-	TEXT("If this is anything other than 0.0 then we can potentially go over our time budget."),
-	FConsoleVariableDelegate::CreateLambda([](IConsoleVariable* InVariable)
-	{
-		GMinQuality = FMath::Clamp(GMinQuality, 0.0f, 1.0f);
-	}),
-	ECVF_Scalability);
-
-static int32 GMaxTickRate = 10;
-
-static FAutoConsoleVariableRef CVarSkelBatch_MaxTickRate(
-	TEXT("a.Budget.MaxTickRate"),
-	GMaxTickRate,
-	TEXT("Values >= 1, Default = 10\n")
-	TEXT("The maximum tick rate we allow. If this is set then we can potentially go over budget, but keep quality of individual meshes to a reasonable level.\n"),
-	FConsoleVariableDelegate::CreateLambda([](IConsoleVariable* InVariable)
-	{
-		GMaxTickRate = FMath::Max(GMaxTickRate, 1);
-	}),
-	ECVF_Scalability);
-
-static float GWorkUnitSmoothingSpeed = 5.0f;
-
-static FAutoConsoleVariableRef CVarSkelBatch_WorkUnitSmoothingSpeed(
-	TEXT("a.Budget.WorkUnitSmoothingSpeed"),
-	GWorkUnitSmoothingSpeed,
-	TEXT("Values > 0.1, Default = 5.0\n")
-	TEXT("The speed at which the average work unit converges on the measured amount.\n"),
-	FConsoleVariableDelegate::CreateLambda([](IConsoleVariable* InVariable)
-	{
-		GWorkUnitSmoothingSpeed = FMath::Max(GWorkUnitSmoothingSpeed, 0.1f);
-	}));
-
-static float GAlwaysTickFalloffAggression = 0.8f;
-
-static FAutoConsoleVariableRef CVarSkelBatch_AlwaysTickFalloffAggression(
-	TEXT("a.Budget.AlwaysTickFalloffAggression"),
-	GAlwaysTickFalloffAggression,
-	TEXT("Range [0.1, 0.9], Default = 0.8\n")
-	TEXT("Controls the rate at which 'always ticked' components falloff under load.\n")
-	TEXT("Higher values mean that we reduce the number of always ticking components by a larger amount when the allocated time budget is exceeded."),
-	FConsoleVariableDelegate::CreateLambda([](IConsoleVariable* InVariable)
-	{
-		GAlwaysTickFalloffAggression = FMath::Clamp(GAlwaysTickFalloffAggression, 0.1f, 0.9f);
-	}),
-	ECVF_Scalability);
-
-static float GInterpolationFalloffAggression = 0.4f;
-
-static FAutoConsoleVariableRef CVarSkelBatch_InterpolationFalloffAggression(
-	TEXT("a.Budget.InterpolationFalloffAggression"),
-	GInterpolationFalloffAggression,
-	TEXT("Range [0.1, 0.9], Default = 0.4\n")
-	TEXT("Controls the rate at which interpolated components falloff under load.\n")
-	TEXT("Higher values mean that we reduce the number of interpolated components by a larger amount when the allocated time budget is exceeded.\n")
-	TEXT("Components are only interpolated when the time budget is exceeded."),
-	FConsoleVariableDelegate::CreateLambda([](IConsoleVariable* InVariable)
-	{
-		GInterpolationFalloffAggression = FMath::Clamp(GInterpolationFalloffAggression, 0.1f, 0.9f);
-	}),
-	ECVF_Scalability);
-
-static int32 GInterpolationMaxRate = 6;
-
-static FAutoConsoleVariableRef CVarSkelBatch_InterpolationMaxRate(
-	TEXT("a.Budget.InterpolationMaxRate"),
-	GInterpolationMaxRate,
-	TEXT("Values > 1, Default = 6\n")
-	TEXT("Controls the rate at which ticks happen when interpolating.\n"),
-	FConsoleVariableDelegate::CreateLambda([](IConsoleVariable* InVariable)
-	{
-		GInterpolationMaxRate = FMath::Max(GInterpolationMaxRate, 2);
-	}),
-	ECVF_Scalability);
-
-static int32 GMaxInterpolatedComponents = 16;
-
-static FAutoConsoleVariableRef CVarSkelBatch_MaxInterpolatedComponents(
-	TEXT("a.Budget.MaxInterpolatedComponents"),
-	GMaxInterpolatedComponents,
-	TEXT("Range >= 0, Default = 16\n")
-	TEXT("Max number of components to inteprolate before we start throttling.\n"),
-	FConsoleVariableDelegate::CreateLambda([](IConsoleVariable* InVariable)
-	{
-		GMaxInterpolatedComponents = FMath::Max(GMaxInterpolatedComponents, 0);
-	}),
-	ECVF_Scalability);
-
-static float GInterpolationTickMultiplier = 0.75f;
-
-static FAutoConsoleVariableRef CVarSkelBatch_InterpolationTickMultiplier(
-	TEXT("a.Budget.InterpolationTickMultiplier"),
-	GInterpolationTickMultiplier,
-	TEXT("Range [0.1, 0.9], Default = 0.75\n")
-	TEXT("Controls the expected value an amortized interpolated tick will take compared to a 'normal' tick.\n"),
-	FConsoleVariableDelegate::CreateLambda([](IConsoleVariable* InVariable)
-	{
-		GInterpolationTickMultiplier = FMath::Clamp(GInterpolationTickMultiplier, 0.1f, 0.9f);
-	}),
-	ECVF_Scalability);
-
-static float GInitialEstimatedWorkUnitTimeMs = 0.08f;
-
-static FAutoConsoleVariableRef CVarSkelBatch_InitialEstimatedWorkUnitTime(
-	TEXT("a.Budget.InitialEstimatedWorkUnitTime"),
-	GInitialEstimatedWorkUnitTimeMs,
-	TEXT("Values > 0.0, Default = 0.08\n")
-	TEXT("Controls the time in milliseconds we expect, on average, for a skeletal mesh component to execute.\n")
-	TEXT("The value only applies for the first tick of a component, after which we use the real time the tick takes.\n"),
-	FConsoleVariableDelegate::CreateLambda([](IConsoleVariable* InVariable)
-	{
-		GInitialEstimatedWorkUnitTimeMs = FMath::Max(GInitialEstimatedWorkUnitTimeMs, KINDA_SMALL_NUMBER);
-	}),
-	ECVF_Scalability);
-
-static int32 GMaxTickedOffsreenComponents = 4;
-
-static FAutoConsoleVariableRef CVarSkelBatch_MaxTickedOffsreenComponents(
-	TEXT("a.Budget.MaxTickedOffsreen"),
-	GMaxTickedOffsreenComponents,
-	TEXT("Values >= 1, Default = 4\n")
-	TEXT("The maximum number of offscreen components we tick (most significant first)\n"),
-	FConsoleVariableDelegate::CreateLambda([](IConsoleVariable* InVariable)
-	{
-		GMaxTickedOffsreenComponents = FMath::Max(GMaxTickedOffsreenComponents, 1);
-	}),
-	ECVF_Scalability);
-
-static int32 GStateChangeThrottleInFrames = 30;
-
-static FAutoConsoleVariableRef CVarSkelBatch_StateChangeThrottleInFrames(
-	TEXT("a.Budget.StateChangeThrottleInFrames"),
-	GStateChangeThrottleInFrames,
-	TEXT("Range [1, 255], Default = 30\n")
-	TEXT("Prevents throttle values from changing too often due to system and load noise.\n"),
-	FConsoleVariableDelegate::CreateLambda([](IConsoleVariable* InVariable)
-	{
-		GStateChangeThrottleInFrames = FMath::Clamp(GStateChangeThrottleInFrames, 1, 255);
-	}),
-	ECVF_Scalability);
-
-static float GBudgetFactorBeforeReducedWork = 1.5f;
-
-static FAutoConsoleVariableRef CVarSkelBatch_BudgetFactorBeforeReducedWork(
-	TEXT("a.Budget.BudgetFactorBeforeReducedWork"),
-	GBudgetFactorBeforeReducedWork,
-	TEXT("Range > 1, Default = 1.5\n")
-	TEXT("Reduced work will be delayed until budget pressure goes over this amount.\n"),
-	FConsoleVariableDelegate::CreateLambda([](IConsoleVariable* InVariable)
-	{
-		GBudgetFactorBeforeReducedWork = FMath::Max(GBudgetFactorBeforeReducedWork, 1.0f);
-	}),
-	ECVF_Scalability);
-
-static float GBudgetFactorBeforeReducedWorkEpsilon = 0.25f;
-
-static FAutoConsoleVariableRef CVarSkelBatch_BudgetFactorBeforeReducedWorkEpsilon(
-	TEXT("a.Budget.BudgetFactorBeforeReducedWorkEpsilon"),
-	GBudgetFactorBeforeReducedWorkEpsilon,
-	TEXT("Range > 0.0, Default = 0.25\n")
-	TEXT("Increased work will be delayed until budget pressure goes under a.Budget.BudgetFactorBeforeReducedWork minus this amount.\n"),
-	FConsoleVariableDelegate::CreateLambda([](IConsoleVariable* InVariable)
-	{
-		GBudgetFactorBeforeReducedWorkEpsilon = FMath::Max(GBudgetFactorBeforeReducedWorkEpsilon, 0.0f);
-	}),
-	ECVF_Scalability);
-
-static float GBudgetPressureSmoothingSpeed = 3.0f;
-
-static FAutoConsoleVariableRef CVarSkelBatch_BudgetPressureSmoothingSpeed(
-	TEXT("a.Budget.BudgetPressureSmoothingSpeed"),
-	GBudgetPressureSmoothingSpeed,
-	TEXT("Range > 0.0, Default = 3.0\n")
-	TEXT("How much to smooth the budget pressure value used to throttle reduced work.\n"),
-	FConsoleVariableDelegate::CreateLambda([](IConsoleVariable* InVariable)
-	{
-		GBudgetPressureSmoothingSpeed = FMath::Max(GBudgetPressureSmoothingSpeed, KINDA_SMALL_NUMBER);
-	}),
-	ECVF_Scalability);
-
-static int32 GReducedWorkThrottleMinInFrames = 2;
-
-static FAutoConsoleVariableRef CVarSkelBatch_ReducedWorkThrottleMinInFrames(
-	TEXT("a.Budget.ReducedWorkThrottleMinInFrames"),
-	GReducedWorkThrottleMinInFrames,
-	TEXT("Range [1, 255], Default = 2\n")
-	TEXT("Prevents reduced work from changing too often due to system and load noise. Min value used when over budget pressure (i.e. aggressive reduction).\n"),
-	FConsoleVariableDelegate::CreateLambda([](IConsoleVariable* InVariable)
-	{
-		GReducedWorkThrottleMinInFrames = FMath::Clamp(GReducedWorkThrottleMinInFrames, 1, 255);
-	}),
-	ECVF_Scalability);
-
-static int32 GReducedWorkThrottleMaxInFrames = 20;
-
-static FAutoConsoleVariableRef CVarSkelBatch_ReducedWorkThrottleMaxInFrames(
-	TEXT("a.Budget.ReducedWorkThrottleMaxInFrames"),
-	GReducedWorkThrottleMaxInFrames,
-	TEXT("Range [1, 255], Default = 20\n")
-	TEXT("Prevents reduced work from changing too often due to system and load noise. Max value used when under budget pressure.\n"),
-	FConsoleVariableDelegate::CreateLambda([](IConsoleVariable* InVariable)
-	{
-		GReducedWorkThrottleMaxInFrames = FMath::Clamp(GReducedWorkThrottleMaxInFrames, 1, 255);
-	}),
-	ECVF_Scalability);
-
-static float GBudgetFactorBeforeAggressiveReducedWork = 2.0f;
-
-static FAutoConsoleVariableRef CVarSkelBatch_BudgetFactorBeforeAggressiveReducedWork(
-	TEXT("a.Budget.BudgetFactorBeforeAggressiveReducedWork"),
-	GBudgetFactorBeforeAggressiveReducedWork,
-	TEXT("Range > 1, Default = 2.0\n")
-	TEXT("Reduced work will be applied more rapidly when budget pressure goes over this amount.\n"),
-	FConsoleVariableDelegate::CreateLambda([](IConsoleVariable* InVariable)
-	{
-		GBudgetFactorBeforeAggressiveReducedWork = FMath::Max(GBudgetFactorBeforeAggressiveReducedWork, 1.0f);
-	}),
-	ECVF_Scalability);
-
-static int32 GReducedWorkThrottleMaxPerFrame = 4;
-
-static FAutoConsoleVariableRef CVarSkelBatch_ReducedWorkThrottleMaxPerFrame(
-	TEXT("a.Budget.ReducedWorkThrottleMaxPerFrame"),
-	GReducedWorkThrottleMaxPerFrame,
-	TEXT("Range [1, 255], Default = 4\n")
-	TEXT("Controls the max number of components that are switched to/from reduced work per tick.\n"),
-	FConsoleVariableDelegate::CreateLambda([](IConsoleVariable* InVariable)
-	{
-		GReducedWorkThrottleMaxPerFrame = FMath::Clamp(GReducedWorkThrottleMaxPerFrame, 1, 255);
-	}),
-	ECVF_Scalability);
-
-static float GBudgetPressureBeforeEmergencyReducedWork = 2.5f;
-
-static FAutoConsoleVariableRef CVarSkelBatch_BudgetPressureBeforeEmergencyReducedWork(
-	TEXT("a.Budget.GBudgetPressureBeforeEmergencyReducedWork"),
-	GBudgetPressureBeforeEmergencyReducedWork,
-	TEXT("Range > 0.0, Default = 2.5\n")
-	TEXT("Controls the budget pressure where emergency reduced work (applied to all components except those that are bAlwaysTick).\n"),
-	FConsoleVariableDelegate::CreateLambda([](IConsoleVariable* InVariable)
-	{
-		GBudgetPressureBeforeEmergencyReducedWork = FMath::Max(GBudgetPressureBeforeEmergencyReducedWork, 0.0f);
-	}),
-	ECVF_Scalability);
-
-FComponentData::FComponentData(USkeletalMeshComponentBudgeted* InComponent)
+FComponentData::FComponentData(USkeletalMeshComponentBudgeted* InComponent, float InGameThreadLastTickTimeMs, int32 InStateChangeThrottle)
 	: Component(InComponent)
 	, RootPrerequisite(nullptr)
 	, Significance(1.0f)
 	, AccumulatedDeltaTime(0.0f)
-	, GameThreadLastTickTimeMs(GInitialEstimatedWorkUnitTimeMs)
+	, GameThreadLastTickTimeMs(InGameThreadLastTickTimeMs)
 	, GameThreadLastCompletionTimeMs(0.0f)
 	, FrameOffset(0)
 	, DesiredTickRate(1)
 	, TickRate(1)
 	, SkippedTicks(0)
-	, StateChangeThrottle(GStateChangeThrottleInFrames)
+	, StateChangeThrottle(InStateChangeThrottle)
 	, bTickEnabled(true)
 	, bAlwaysTick(false)
 	, bTickEvenIfNotRendered(false)
@@ -327,24 +58,38 @@ FComponentData::FComponentData(USkeletalMeshComponentBudgeted* InComponent)
 
 FAnimationBudgetAllocator::FAnimationBudgetAllocator(UWorld* InWorld)
 	: World(InWorld)
-	, AverageWorkUnitTimeMs(GInitialEstimatedWorkUnitTimeMs)
+	, AverageWorkUnitTimeMs(GBudgetParameters.InitialEstimatedWorkUnitTimeMs)
 	, NumComponentsToNotSkip(0)
 	, TotalEstimatedTickTimeMs(0.0f)
 	, NumWorkUnitsForAverage(0.0f)
 	, SmoothedBudgetPressure(0.0f)
+#if ENABLE_DRAW_DEBUG
+	, DebugTotalTime(0.0f)
+	, CurrentDebugTimeDisplay(0.0f)
+	, DebugSmoothedTotalTime(0.0f)
+#endif
 	, ReducedComponentWorkCounter(0)
 	, CurrentFrameOffset(0)
 	, bEnabled(false)
 {
-	FAnimationBudgetAllocator::bCachedEnabled = GAnimationBudgetEnabled == 1 && bEnabled;
+	SetParametersFromCVars();
+
 	PostGarbageCollectHandle = FCoreUObjectDelegates::GetPostGarbageCollect().AddRaw(this, &FAnimationBudgetAllocator::HandlePostGarbageCollect);
 	OnWorldPreActorTickHandle = FWorldDelegates::OnWorldPreActorTick.AddRaw(this, &FAnimationBudgetAllocator::OnWorldPreActorTick);
+	OnCVarParametersChangedHandle = GOnCVarParametersChanged.AddRaw(this, &FAnimationBudgetAllocator::SetParametersFromCVars);
+#if ENABLE_DRAW_DEBUG
+	OnHUDPostRenderHandle = AHUD::OnHUDPostRender.AddRaw(this, &FAnimationBudgetAllocator::OnHUDPostRender);
+#endif
 }
 
 FAnimationBudgetAllocator::~FAnimationBudgetAllocator()
 {
+#if ENABLE_DRAW_DEBUG
+	AHUD::OnHUDPostRender.Remove(OnHUDPostRenderHandle);
+#endif
 	FCoreUObjectDelegates::GetPostGarbageCollect().Remove(PostGarbageCollectHandle);
 	FWorldDelegates::OnWorldPreActorTick.Remove(OnWorldPreActorTickHandle);
+	GOnCVarParametersChanged.Remove(OnCVarParametersChangedHandle);
 }
 
 IAnimationBudgetAllocator* IAnimationBudgetAllocator::Get(UWorld* InWorld)
@@ -354,8 +99,7 @@ IAnimationBudgetAllocator* IAnimationBudgetAllocator::Get(UWorld* InWorld)
 
 void FAnimationBudgetAllocator::SetComponentTickEnabled(USkeletalMeshComponentBudgeted* Component, bool bShouldTick)
 {
-#if USE_SKEL_BATCHING
-	if (FAnimationBudgetAllocator::bCachedEnabled)
+	if (FAnimationBudgetAllocator::bCachedEnabled && bEnabled)
 	{
 		int32 Handle = Component->GetAnimationBudgetHandle();
 		if(Handle != INDEX_NONE)
@@ -366,7 +110,6 @@ void FAnimationBudgetAllocator::SetComponentTickEnabled(USkeletalMeshComponentBu
 		TickEnableHelper(Component, bShouldTick);
 	}
 	else
-#endif
 	{
 		TickEnableHelper(Component, bShouldTick);
 	}
@@ -374,8 +117,7 @@ void FAnimationBudgetAllocator::SetComponentTickEnabled(USkeletalMeshComponentBu
 
 bool FAnimationBudgetAllocator::IsComponentTickEnabled(USkeletalMeshComponentBudgeted* Component) const
 {
-#if USE_SKEL_BATCHING
-	if (FAnimationBudgetAllocator::bCachedEnabled)
+	if (FAnimationBudgetAllocator::bCachedEnabled && bEnabled)
 	{
 		int32 Handle = Component->GetAnimationBudgetHandle();
 		if(Handle != INDEX_NONE)
@@ -386,7 +128,6 @@ bool FAnimationBudgetAllocator::IsComponentTickEnabled(USkeletalMeshComponentBud
 		return Component->PrimaryComponentTick.IsTickFunctionEnabled();
 	}
 	else
-#endif
 	{
 		return Component->PrimaryComponentTick.IsTickFunctionEnabled();
 	}
@@ -394,8 +135,7 @@ bool FAnimationBudgetAllocator::IsComponentTickEnabled(USkeletalMeshComponentBud
 
 void FAnimationBudgetAllocator::SetComponentSignificance(USkeletalMeshComponentBudgeted* Component, float Significance, bool bAlwaysTick, bool bTickEvenIfNotRendered, bool bAllowReducedWork, bool bNeverThrottle)
 {
-#if USE_SKEL_BATCHING
-	if (FAnimationBudgetAllocator::bCachedEnabled)
+	if (FAnimationBudgetAllocator::bCachedEnabled && bEnabled)
 	{
 		int32 Handle = Component->GetAnimationBudgetHandle();
 		if(Handle != INDEX_NONE)
@@ -408,7 +148,6 @@ void FAnimationBudgetAllocator::SetComponentSignificance(USkeletalMeshComponentB
 			ComponentData.bNeverThrottle = bNeverThrottle;
 		}
 	}
-#endif
 }
 
 void FAnimationBudgetAllocator::TickEnableHelper(USkeletalMeshComponent* InComponent, bool bInEnable)
@@ -566,7 +305,7 @@ void FAnimationBudgetAllocator::QueueSortedComponentIndices(float InDeltaSeconds
 	ReducedWorkComponentData.Sort(SignificanceSortPredicate);
 	NonRenderedComponentData.Sort(SignificanceSortPredicate);
 
-	const int32 MaxOffscreenComponents = FMath::Min(NonRenderedComponentData.Num(), GMaxTickedOffsreenComponents);
+	const int32 MaxOffscreenComponents = FMath::Min(NonRenderedComponentData.Num(), Parameters.MaxTickedOffsreenComponents);
 	if(MaxOffscreenComponents > 0)
 	{
 		auto ReduceWorkForOffscreenComponent = [](FComponentData& InComponentData)
@@ -680,20 +419,20 @@ int32 FAnimationBudgetAllocator::CalculateWorkDistributionAndQueue(float InDelta
 	{
 		// Calc smoothed average of last frames' work units
 		const float AverageTickTimeMs = TotalEstimatedTickTimeMs / NumWorkUnitsForAverage;
-		AverageWorkUnitTimeMs = FMath::FInterpTo(AverageWorkUnitTimeMs, AverageTickTimeMs, InDeltaSeconds, GWorkUnitSmoothingSpeed);
+		AverageWorkUnitTimeMs = FMath::FInterpTo(AverageWorkUnitTimeMs, AverageTickTimeMs, InDeltaSeconds, Parameters.WorkUnitSmoothingSpeed);
 
 		SET_FLOAT_STAT(STAT_AnimationBudgetAllocator_AverageWorkUnitTime, AverageWorkUnitTimeMs);
 		CSV_CUSTOM_STAT(AnimationBudget, AverageWorkUnitTimeMs, AverageTickTimeMs, ECsvCustomStatOp::Set);
 
 		// Want to map the remaining (non-fixed) work units so that we only execute N work units per frame.
 		// If we can go over budget to keep quality then we use that value
-		const float WorkUnitBudget = FMath::Max(GBudgetInMs / AverageWorkUnitTimeMs, (float)TotalIdealWorkUnits * GMinQuality);
+		const float WorkUnitBudget = FMath::Max(Parameters.BudgetInMs / AverageWorkUnitTimeMs, (float)TotalIdealWorkUnits * Parameters.MinQuality);
 
 		SET_FLOAT_STAT(STAT_AnimationBudgetAllocator_Budget, WorkUnitBudget);
 
 		// Ramp-off work units that we tick every frame once required ticks start exceeding budget
 		const float WorkUnitsExcess = FMath::Max(0.0f, TotalIdealWorkUnits - WorkUnitBudget);
-		const float WorkUnitsToRunInFull = FMath::Clamp(WorkUnitBudget - (WorkUnitsExcess * GAlwaysTickFalloffAggression), (float)NumComponentsToNotSkip, (float)TotalIdealWorkUnits);
+		const float WorkUnitsToRunInFull = FMath::Clamp(WorkUnitBudget - (WorkUnitsExcess * Parameters.AlwaysTickFalloffAggression), (float)NumComponentsToNotSkip, (float)TotalIdealWorkUnits);
 		SET_DWORD_STAT(STAT_AnimationBudgetAllocator_AlwaysTick, WorkUnitsToRunInFull);
 		BUDGET_CSV_STAT(AnimationBudget, NumAlwaysTicked, WorkUnitsToRunInFull, ECsvCustomStatOp::Set);
 		const int32 FullIndexEnd = (int32)WorkUnitsToRunInFull;
@@ -714,16 +453,16 @@ int32 FAnimationBudgetAllocator::CalculateWorkDistributionAndQueue(float InDelta
 		float RemainingWorkUnitsToRun = FMath::Max(0.0f, TotalIdealWorkUnits - FullTickWorkUnits);
 
 		// Ramp off interpolated units in a similar way
-		const float WorkUnitsToInterpolate = FMath::Min(FMath::Max(RemainingBudget - (WorkUnitsExcess * GInterpolationFalloffAggression), (float)FMath::Min(GMaxInterpolatedComponents, NumComponentsToNotThrottle)), RemainingWorkUnitsToRun);
+		const float WorkUnitsToInterpolate = FMath::Min(FMath::Max(RemainingBudget - (WorkUnitsExcess * Parameters.InterpolationFalloffAggression), (float)FMath::Min(Parameters.MaxInterpolatedComponents, NumComponentsToNotThrottle)), RemainingWorkUnitsToRun);
 		SET_DWORD_STAT(STAT_AnimationBudgetAllocator_Interpolated, WorkUnitsToInterpolate);
 
 		const int32 InterpolationIndexEnd = FMath::Min((int32)WorkUnitsToInterpolate + (int32)WorkUnitsToRunInFull, TotalIdealWorkUnits);
 
-		const float MaxInterpolationRate = (float)GInterpolationMaxRate;
+		const float MaxInterpolationRate = (float)Parameters.InterpolationMaxRate;
 
 		// Calc remaining (throttled) work units
-		RemainingBudget = FMath::Max(0.0f, RemainingBudget - (WorkUnitsToInterpolate * GInterpolationTickMultiplier));
-		RemainingWorkUnitsToRun = FMath::Max(0.0f, RemainingWorkUnitsToRun - (WorkUnitsToInterpolate * GInterpolationTickMultiplier));
+		RemainingBudget = FMath::Max(0.0f, RemainingBudget - (WorkUnitsToInterpolate * Parameters.InterpolationTickMultiplier));
+		RemainingWorkUnitsToRun = FMath::Max(0.0f, RemainingWorkUnitsToRun - (WorkUnitsToInterpolate * Parameters.InterpolationTickMultiplier));
 
 		SET_DWORD_STAT(STAT_AnimationBudgetAllocator_Throttled, RemainingWorkUnitsToRun);
 
@@ -733,7 +472,7 @@ int32 FAnimationBudgetAllocator::CalculateWorkDistributionAndQueue(float InDelta
 		// so we keep the area under the curve constant and intercept the line with this centroid.
 		// Care must be taken with rounding to keep workload in-budget.
 		const float ThrottleRateDenominator = RemainingBudget > 1.0f ? RemainingBudget : 1.0f;
-		const float MaxThrottleRate = FMath::Min(FMath::CeilToFloat(FMath::Max(1.0f, RemainingWorkUnitsToRun / ThrottleRateDenominator) * 2.0f), (float)GMaxTickRate);
+		const float MaxThrottleRate = FMath::Min(FMath::CeilToFloat(FMath::Max(1.0f, RemainingWorkUnitsToRun / ThrottleRateDenominator) * 2.0f), (float)Parameters.MaxTickRate);
 		const float ThrottleDenominator = RemainingWorkUnitsToRun > 0.0f ? RemainingWorkUnitsToRun : 1.0f;
 
 		// Bucket 1: always ticked
@@ -775,10 +514,10 @@ int32 FAnimationBudgetAllocator::CalculateWorkDistributionAndQueue(float InDelta
 		BUDGET_CSV_STAT(AnimationBudget, NumThrottled, RemainingWorkUnitsToRun, ECsvCustomStatOp::Set);
 
 		const float BudgetPressure = (float)TotalIdealWorkUnits / WorkUnitBudget;
-		SmoothedBudgetPressure = FMath::FInterpTo(SmoothedBudgetPressure, BudgetPressure, InDeltaSeconds, GBudgetPressureSmoothingSpeed);
+		SmoothedBudgetPressure = FMath::FInterpTo(SmoothedBudgetPressure, BudgetPressure, InDeltaSeconds, Parameters.BudgetPressureSmoothingSpeed);
 
-		float BudgetPressureInterpAlpha = FMath::Clamp((SmoothedBudgetPressure - GBudgetFactorBeforeAggressiveReducedWork) * 0.5f, 0.0f, 1.0f);
-		int32 StateChangeThrottleInFrames = (int32)FMath::Lerp(4.0f, (float)GStateChangeThrottleInFrames, BudgetPressureInterpAlpha);
+		float BudgetPressureInterpAlpha = FMath::Clamp((SmoothedBudgetPressure - Parameters.BudgetFactorBeforeAggressiveReducedWork) * 0.5f, 0.0f, 1.0f);
+		int32 StateChangeThrottleInFrames = (int32)FMath::Lerp(4.0f, (float)Parameters.StateChangeThrottleInFrames, BudgetPressureInterpAlpha);
 
 		SET_FLOAT_STAT(STAT_AnimationBudgetAllocator_SmoothedBudgetPressure, SmoothedBudgetPressure);
 
@@ -819,14 +558,14 @@ int32 FAnimationBudgetAllocator::CalculateWorkDistributionAndQueue(float InDelta
 
 		if(--ReducedComponentWorkCounter <= 0)
 		{
-			const bool bEmergencyReducedWork = SmoothedBudgetPressure >= GBudgetPressureBeforeEmergencyReducedWork;
+			const bool bEmergencyReducedWork = SmoothedBudgetPressure >= Parameters.BudgetPressureBeforeEmergencyReducedWork;
 
 			// Scale num components to switch based on budget pressure
-			const int32 NumComponentsToSwitch = (int32)FMath::Lerp(1.0f, (float)GReducedWorkThrottleMaxPerFrame, BudgetPressureInterpAlpha);
+			const int32 NumComponentsToSwitch = (int32)FMath::Lerp(1.0f, (float)Parameters.ReducedWorkThrottleMaxPerFrame, BudgetPressureInterpAlpha);
 			int32 ComponentsSwitched = 0;
 
 			// If we have any components running reduced work when we have an excess, then move them out of the 'reduced' pool per tick
-			if (ReducedWorkComponentData.Num() > 0 && SmoothedBudgetPressure < GBudgetFactorBeforeReducedWork - GBudgetFactorBeforeReducedWorkEpsilon)
+			if (ReducedWorkComponentData.Num() > 0 && SmoothedBudgetPressure < Parameters.BudgetFactorBeforeReducedWork - Parameters.BudgetFactorBeforeReducedWorkEpsilon)
 			{
 				for(int32 ReducedWorkComponentIndex : ReducedWorkComponentData)
 				{
@@ -847,7 +586,7 @@ int32 FAnimationBudgetAllocator::CalculateWorkDistributionAndQueue(float InDelta
 					}
 				}
 			}
-			else if(SmoothedBudgetPressure > GBudgetFactorBeforeReducedWork)
+			else if(SmoothedBudgetPressure > Parameters.BudgetFactorBeforeReducedWork)
 			{
 				// Any work units that we interpolate or throttle should also be eligible for work reduction (which can involve disabling other ticks), so set them all now if needed
 				for (SortedComponentIndex = TotalIdealWorkUnits - 1; SortedComponentIndex >= FullIndexEnd; --SortedComponentIndex)
@@ -874,7 +613,7 @@ int32 FAnimationBudgetAllocator::CalculateWorkDistributionAndQueue(float InDelta
 			}
 
 			// Scale the rate at which we consider reducing component work based on budget pressure
-			ReducedComponentWorkCounter = (int32)FMath::Lerp((float)GReducedWorkThrottleMaxInFrames, (float)GReducedWorkThrottleMinInFrames, BudgetPressureInterpAlpha);
+			ReducedComponentWorkCounter = (int32)FMath::Lerp((float)Parameters.ReducedWorkThrottleMaxInFrames, (float)Parameters.ReducedWorkThrottleMinInFrames, BudgetPressureInterpAlpha);
 		}
 	}
 
@@ -907,10 +646,9 @@ void FAnimationBudgetAllocator::Update(float DeltaSeconds)
 	SCOPE_CYCLE_COUNTER(STAT_AnimationBudgetAllocator_Update);
 	CSV_SCOPED_TIMING_STAT_EXCLUSIVE(AnimationBudgetAllocator);
 
-	FAnimationBudgetAllocator::bCachedEnabled = GAnimationBudgetEnabled == 1 && bEnabled;
+	FAnimationBudgetAllocator::bCachedEnabled = GAnimationBudgetEnabled == 1;
 
-#if USE_SKEL_BATCHING
-	if (FAnimationBudgetAllocator::bCachedEnabled)
+	if (FAnimationBudgetAllocator::bCachedEnabled && bEnabled)
 	{
 		check(IsInGameThread());
 
@@ -935,17 +673,108 @@ void FAnimationBudgetAllocator::Update(float DeltaSeconds)
 		BUDGET_CSV_STAT(AnimationBudget, AnimQuality, AllSortedComponentData.Num() > 0 ? (float)NumTicked / (float)AllSortedComponentData.Num() : 0.0f, ECsvCustomStatOp::Set);
 		BUDGET_CSV_STAT(AnimationBudget, AverageTickRate, AverageTickRate, ECsvCustomStatOp::Set);
 
-#if WITH_TICK_DEBUG
-		for (int32 ComponentDataIndex : AllSortedComponentData)
+#if ENABLE_DRAW_DEBUG
+		if(GAnimationBudgetDebugEnabled != 0)
 		{
-			FComponentData& ComponentData = AllComponentData[ComponentDataIndex];
-			DrawDebugString(World, ComponentData.Component->GetOwner()->GetActorLocation(), FString::Printf(TEXT("0x%llx\n%d (%s)\n%s, %s"), &ComponentData, ComponentData.TickRate, ComponentData.bInterpolate ? TEXT("Interp") : TEXT("No Interp"), ComponentData.bReducedWork ? TEXT("Reduced") : TEXT("NotReduced"), ComponentData.bAllowReducedWork ? TEXT("AllowReduced") : TEXT("DisallowReduced")), nullptr, FColor::White, 0.016f, false);
+			TMap<AActor*, TArray<int32>> ActorMap;
+			for (int32 ComponentDataIndex : AllSortedComponentData)
+			{
+				FComponentData& ComponentData = AllComponentData[ComponentDataIndex];
+				TArray<int32>& ComponentIndexArray = ActorMap.FindOrAdd(ComponentData.Component->GetOwner());
+				ComponentIndexArray.Add(ComponentDataIndex);
+			}
+
+			for(const TPair<AActor*, TArray<int32>>& ActorIndicesPair : ActorMap)
+			{
+				FVector Location = ActorIndicesPair.Key->GetActorLocation();
+
+				FString DebugString;
+				
+				for(int32 ComponentDataIndex : ActorIndicesPair.Value)
+				{
+					FComponentData& ComponentData = AllComponentData[ComponentDataIndex];
+					if(ComponentData.bTickEnabled)
+					{
+						if(GAnimationBudgetDebugShowAddresses != 0)
+						{
+							DebugString += FString::Printf(TEXT("0x%llx %d %s %s\n"), &ComponentData, ComponentData.TickRate, ComponentData.bInterpolate ? TEXT("I") : TEXT(" "), ComponentData.bReducedWork ? TEXT("Lo") : TEXT("Hi"));
+						}
+						else
+						{
+							DebugString += FString::Printf(TEXT("%d %s %s\n"), ComponentData.TickRate, ComponentData.bInterpolate ? TEXT("I") : TEXT(" "), ComponentData.bReducedWork ? TEXT("Lo") : TEXT("Hi"));
+						}
+					}
+				}
+
+				DrawDebugString(World, Location, DebugString, nullptr, FColor::White, 0.016f, false);
+			}
+
+			DebugTimes.Add(FVector2D(CurrentDebugTimeDisplay, DebugTotalTime));
+
+			float DebugTimeAccumulator = 0.0f;
+			float DebugTimeNum = 0.0f;
+			for(int32 DebugTimeIndex = DebugTimes.Num() - 1, DebugTimeCount = 0; DebugTimeIndex >= 0 && DebugTimeCount < 20; DebugTimeIndex--, DebugTimeCount++)
+			{
+				DebugTimeAccumulator += DebugTimes[DebugTimeIndex].Y;
+				DebugTimeNum += 1.0f;
+			}
+			DebugTimesSmoothed.Add(FVector2D(CurrentDebugTimeDisplay, DebugTimeNum > 0.0f ? DebugTimeAccumulator / DebugTimeNum : 0.0f));
+
+			CurrentDebugTimeDisplay += 0.016f;
+			DebugSmoothedTotalTime = FMath::Max(FMath::CeilToFloat(DebugTimesSmoothed.Last().Y * 0.5f) / 0.5f, Parameters.BudgetInMs);
+			DebugTotalTime = 0.0f;
 		}
 #endif
 	}
-#endif
 }
 
+#if ENABLE_DRAW_DEBUG
+void FAnimationBudgetAllocator::OnHUDPostRender(AHUD* HUD, UCanvas* Canvas)
+{
+	if(FAnimationBudgetAllocator::bCachedEnabled && bEnabled && GAnimationBudgetDebugEnabled != 0)
+	{
+		if(HUD->GetWorld() == World)
+		{
+			TWeakObjectPtr<UReporterGraph> Graph = Canvas->GetReporterGraph();
+
+			Graph->SetNumGraphLines(2);
+
+			FGraphLine* RawGraphLine = Graph->GetGraphLine(0);
+
+			RawGraphLine->Color = FLinearColor::Gray.CopyWithNewOpacity(0.3f);
+			RawGraphLine->LineName = TEXT("Raw");
+			RawGraphLine->Data = DebugTimes;
+
+			FGraphLine* SmoothedGraphLine = Graph->GetGraphLine(1);
+
+			SmoothedGraphLine->Color = FLinearColor::White;
+			SmoothedGraphLine->LineName = TEXT("Smoothed");
+			SmoothedGraphLine->Data = DebugTimesSmoothed;
+
+			Graph->SetNumThresholds(1);
+
+			FGraphThreshold* Threshold = Graph->GetThreshold(0);
+			Threshold->Color = FLinearColor::White;
+			Threshold->ThresholdName = TEXT("Budget");
+			Threshold->Threshold = Parameters.BudgetInMs;
+
+			Graph->SetNotchesPerAxis(20, 3);
+			Graph->SetAxesMinMax(CurrentDebugTimeDisplay - 10.0f, CurrentDebugTimeDisplay, 0.0f, DebugSmoothedTotalTime * 1.5f);
+			Graph->SetGraphScreenSize(0.1f, 0.9f, 0.1f, 0.3f);
+			Graph->DrawCursorOnGraph(false);
+			Graph->UseTinyFont(true);
+			Graph->SetCursorLocation(CurrentDebugTimeDisplay);
+			Graph->SetStyles(EGraphAxisStyle::Grid, EGraphDataStyle::Lines);
+			Graph->SetLegendPosition(ELegendPosition::Outside);
+			Graph->OffsetDataSets(false);
+			Graph->DrawExtremesOnGraph(false);
+			Graph->bVisible = true;
+			Graph->Draw(Canvas);
+		}
+	}
+}
+#endif
+
 void FAnimationBudgetAllocator::RemoveHelper(int32 Index)
 {
 	if(AllComponentData.IsValidIndex(Index))
@@ -1001,8 +830,7 @@ static USkeletalMeshComponentBudgeted* FindRootPrerequisite(USkeletalMeshCompone
 
 void FAnimationBudgetAllocator::RegisterComponent(USkeletalMeshComponentBudgeted* InComponent)
 {
-#if USE_SKEL_BATCHING
-	if (FAnimationBudgetAllocator::bCachedEnabled)
+	if (FAnimationBudgetAllocator::bCachedEnabled && bEnabled)
 	{
 		if (InComponent->GetAnimationBudgetHandle() == INDEX_NONE)
 		{
@@ -1011,7 +839,7 @@ void FAnimationBudgetAllocator::RegisterComponent(USkeletalMeshComponentBudgeted
 			InComponent->SetAnimationBudgetHandle(AllComponentData.Num());
 
 			// Setup frame offset
-			FComponentData& ComponentData = AllComponentData.Emplace_GetRef(InComponent);
+			FComponentData& ComponentData = AllComponentData.Emplace_GetRef(InComponent, Parameters.InitialEstimatedWorkUnitTimeMs, Parameters.StateChangeThrottleInFrames);
 			USkeletalMeshComponentBudgeted* RootPrerequisite = FindRootPrerequisite(InComponent);
 			ComponentData.RootPrerequisite = (RootPrerequisite != nullptr && RootPrerequisite != InComponent) ? RootPrerequisite : nullptr;
 			ComponentData.FrameOffset = CurrentFrameOffset++;
@@ -1024,13 +852,11 @@ void FAnimationBudgetAllocator::RegisterComponent(USkeletalMeshComponentBudgeted
 			UpdateComponentTickPrerequsites(InComponent);
 		}
 	}
-#endif
 }
 
 void FAnimationBudgetAllocator::UnregisterComponent(USkeletalMeshComponentBudgeted* InComponent)
 {
-#if USE_SKEL_BATCHING
-	if (FAnimationBudgetAllocator::bCachedEnabled)
+	if (FAnimationBudgetAllocator::bCachedEnabled && bEnabled)
 	{
 		int32 ManagerHandle = InComponent->GetAnimationBudgetHandle();
 		if(ManagerHandle != INDEX_NONE)
@@ -1042,13 +868,11 @@ void FAnimationBudgetAllocator::UnregisterComponent(USkeletalMeshComponentBudget
 			InComponent->SetAnimationBudgetAllocator(nullptr);
 		}
 	}
-#endif
 }
 
 void FAnimationBudgetAllocator::UpdateComponentTickPrerequsites(USkeletalMeshComponentBudgeted* InComponent)
 {
-#if USE_SKEL_BATCHING
-	if (FAnimationBudgetAllocator::bCachedEnabled)
+	if (FAnimationBudgetAllocator::bCachedEnabled && bEnabled)
 	{
 		int32 ManagerHandle = InComponent->GetAnimationBudgetHandle();
 		if(ManagerHandle != INDEX_NONE)
@@ -1058,7 +882,6 @@ void FAnimationBudgetAllocator::UpdateComponentTickPrerequsites(USkeletalMeshCom
 			ComponentData.RootPrerequisite = (RootPrerequisite != nullptr && RootPrerequisite != InComponent) ? RootPrerequisite : nullptr;
 		}
 	}
-#endif
 }
 
 void FAnimationBudgetAllocator::AddReferencedObjects(FReferenceCollector& Collector)
@@ -1098,6 +921,9 @@ void FAnimationBudgetAllocator::SetGameThreadLastTickTimeMs(int32 InManagerHandl
 	{
 		FComponentData& ComponentData = AllComponentData[InManagerHandle];
 		ComponentData.GameThreadLastTickTimeMs = InGameThreadLastTickTimeMs;
+#if ENABLE_DRAW_DEBUG
+		DebugTotalTime += InGameThreadLastTickTimeMs;
+#endif
 	}
 }
 
@@ -1107,12 +933,15 @@ void FAnimationBudgetAllocator::SetGameThreadLastCompletionTimeMs(int32 InManage
 	{
 		FComponentData& ComponentData = AllComponentData[InManagerHandle];
 		ComponentData.GameThreadLastCompletionTimeMs = InGameThreadLastCompletionTimeMs;
+#if ENABLE_DRAW_DEBUG
+		DebugTotalTime += InGameThreadLastCompletionTimeMs;
+#endif
 	}
 }
+
 void FAnimationBudgetAllocator::SetIsRunningReducedWork(USkeletalMeshComponentBudgeted* InComponent, bool bInReducedWork)
 {
-#if USE_SKEL_BATCHING
-	if (FAnimationBudgetAllocator::bCachedEnabled)
+	if (GAnimationBudgetEnabled && bEnabled)
 	{
 		int32 ManagerHandle = InComponent->GetAnimationBudgetHandle();
 		if(ManagerHandle != INDEX_NONE)
@@ -1121,7 +950,6 @@ void FAnimationBudgetAllocator::SetIsRunningReducedWork(USkeletalMeshComponentBu
 			ComponentData.bReducedWork = bInReducedWork;
 		}
 	}
-#endif
 }
 
 void FAnimationBudgetAllocator::SetEnabled(bool bInEnabled)
@@ -1145,11 +973,19 @@ void FAnimationBudgetAllocator::SetEnabled(bool bInEnabled)
 
 		AllComponentData.Reset();
 	}
-
-	FAnimationBudgetAllocator::bCachedEnabled = GAnimationBudgetEnabled == 1 && bEnabled;
 }
 
 bool FAnimationBudgetAllocator::GetEnabled() const
 {
 	return bEnabled;
+}
+
+void FAnimationBudgetAllocator::SetParameters(const FAnimationBudgetAllocatorParameters& InParameters)
+{
+	Parameters = InParameters;
+}
+
+void FAnimationBudgetAllocator::SetParametersFromCVars()
+{
+	Parameters = GBudgetParameters;
 }
\ No newline at end of file
diff --git a/Engine/Plugins/Runtime/AnimationBudgetAllocator/Source/AnimationBudgetAllocator/Private/AnimationBudgetAllocator.h b/Engine/Plugins/Runtime/AnimationBudgetAllocator/Source/AnimationBudgetAllocator/Private/AnimationBudgetAllocator.h
--- a/Engine/Plugins/Runtime/AnimationBudgetAllocator/Source/AnimationBudgetAllocator/Private/AnimationBudgetAllocator.h
+++ b/Engine/Plugins/Runtime/AnimationBudgetAllocator/Source/AnimationBudgetAllocator/Private/AnimationBudgetAllocator.h
@@ -11,10 +11,11 @@
 #include "Engine/EngineTypes.h"
 #include "IAnimationBudgetAllocator.h"
 #include "UObject/GCObject.h"
+#include "AnimationBudgetAllocatorParameters.h"
+#include "DrawDebugHelpers.h"
 
 DECLARE_STATS_GROUP(TEXT("Animation Budget Allocator"), STATGROUP_AnimationBudgetAllocator, STATCAT_Advanced);
 
-#define USE_SKEL_BATCHING 1
 #define WITH_TICK_DEBUG 0
 
 // Enable this define to output more dense CSV stats about the budgeter
@@ -28,6 +29,11 @@ DECLARE_STATS_GROUP(TEXT("Animation Budget Allocator"), STATGROUP_AnimationBudge
 
 class FAnimationBudgetAllocator;
 class USkeletalMeshComponentBudgeted;
+class SAnimationBudgetAllocatorDebug;
+class UGameViewportClient;
+class AHUD;
+class UCanvas;
+class FDebugDisplayInfo;
 
 /** Data for a single component */
 struct FComponentData
@@ -55,7 +61,7 @@ struct FComponentData
 		, bNeverThrottle(true)
 	{}
 
-	FComponentData(USkeletalMeshComponentBudgeted* InComponent);
+	FComponentData(USkeletalMeshComponentBudgeted* InComponent, float InGameThreadLastTickTimeMs, int32 InStateChangeThrottle);
 
 	bool operator==(const FComponentData& InOther) const
 	{
@@ -143,6 +149,7 @@ public:
 	virtual void Update(float DeltaSeconds) override;
 	virtual void SetEnabled(bool bInEnabled) override;
 	virtual bool GetEnabled() const override;
+	virtual void SetParameters(const FAnimationBudgetAllocatorParameters& InParameters) override;
 
 	// FGCObject interface
 	virtual void AddReferencedObjects(FReferenceCollector& Collector) override;
@@ -169,7 +176,18 @@ protected:
 	/** Helper function to enable/disable ticks */
 	void TickEnableHelper(USkeletalMeshComponent* InComponent, bool bInEnable);
 
+	/** Initializes internal parameters from CVars. */
+	void SetParametersFromCVars();
+
+#if ENABLE_DRAW_DEBUG
+	/** Debug support function */
+	void OnHUDPostRender(AHUD* HUD, UCanvas* Canvas);
+#endif
+
 protected:
+	/** All of the parameters we use */
+	FAnimationBudgetAllocatorParameters Parameters;
+
 	// World we are linked to
 	UWorld* World;
 
@@ -192,6 +210,12 @@ protected:
 	/** All non-rendered components we might tick */
 	TArray<int32> NonRenderedComponentData;
 
+#if ENABLE_DRAW_DEBUG
+	/** Recorded debug times */
+	TArray<FVector2D> DebugTimes;
+	TArray<FVector2D> DebugTimesSmoothed;
+#endif
+
 	/** Average time for a work unit in milliseconds (smoothed). Updated each tick. */
 	float AverageWorkUnitTimeMs;
 
@@ -210,6 +234,17 @@ protected:
 	/** Budget pressure value, smoothed to reduce noise in 'reduced work' calculations */
 	float SmoothedBudgetPressure;
 
+#if ENABLE_DRAW_DEBUG
+	/** Track total time this tick for debug display */
+	float DebugTotalTime;
+
+	/** Display time for debug graph */
+	float CurrentDebugTimeDisplay;
+
+	/** Display scale for debug graph */
+	float DebugSmoothedTotalTime;
+#endif
+
 	/** Throttle counter for delaying reduced work */
 	int32 ReducedComponentWorkCounter;
 
@@ -219,6 +254,12 @@ protected:
 	/** Handle used for ticking */
 	FDelegateHandle OnWorldPreActorTickHandle;
 
+	/** Handle used for CVar parameter changes */
+	FDelegateHandle OnCVarParametersChangedHandle;
+
+	/** Handle used for debug drawing */
+	FDelegateHandle OnHUDPostRenderHandle;
+
 	/** Offset used to distribute component ticks */
 	uint32 CurrentFrameOffset;
 
diff --git a/Engine/Plugins/Runtime/AnimationBudgetAllocator/Source/AnimationBudgetAllocator/Private/AnimationBudgetAllocatorCVars.cpp b/Engine/Plugins/Runtime/AnimationBudgetAllocator/Source/AnimationBudgetAllocator/Private/AnimationBudgetAllocatorCVars.cpp
new file mode 100644
--- /dev/null
+++ b/Engine/Plugins/Runtime/AnimationBudgetAllocator/Source/AnimationBudgetAllocator/Private/AnimationBudgetAllocatorCVars.cpp
@@ -0,0 +1,286 @@
+// Copyright 1998-2019 Epic Games, Inc. All Rights Reserved.
+
+#include "AnimationBudgetAllocatorCVars.h"
+#include "AnimationBudgetAllocatorParameters.h"
+#include "HAL/IConsoleManager.h"
+
+int32 GAnimationBudgetEnabled = 0;
+
+static FAutoConsoleVariableRef CVarSkelBatch_Enabled(
+	TEXT("a.Budget.Enabled"),
+	GAnimationBudgetEnabled,
+	TEXT("Values: 0/1\n")
+	TEXT("Controls whether the skeletal mesh batching system is enabled. Should be set when there are no running skeletal meshes."),
+	ECVF_Scalability);
+
+#if ENABLE_DRAW_DEBUG
+int32 GAnimationBudgetDebugEnabled = 0;
+
+static FAutoConsoleVariableRef CVarSkelBatch_DebugEnabled(
+	TEXT("a.Budget.Debug.Enabled"),
+	GAnimationBudgetDebugEnabled,
+	TEXT("Values: 0/1\n")
+	TEXT("Controls whether debug rendering (in builds that support it) is enabled for the animation budget allocator."),
+	ECVF_Scalability);
+
+int32 GAnimationBudgetDebugShowAddresses = 0;
+
+static FAutoConsoleVariableRef CVarSkelBatch_ShowAddresses(
+	TEXT("a.Budget.Debug.ShowAddresses"),
+	GAnimationBudgetDebugShowAddresses,
+	TEXT("Values: 0/1\n")
+	TEXT("Controls whether debug rendering shows addresses of component data for debugging."),
+	ECVF_Scalability);
+#endif
+
+/** CVar-driven parameter block */
+FAnimationBudgetAllocatorParameters GBudgetParameters;
+
+/** Delegate broadcast when parameter block changes */
+FSimpleMulticastDelegate GOnCVarParametersChanged;
+
+static FAutoConsoleVariableRef CVarSkelBatch_Budget(
+	TEXT("a.Budget.BudgetMs"),
+	GBudgetParameters.BudgetInMs,
+	TEXT("Values > 0.1, Default = 1.0\n")
+	TEXT("The time in milliseconds that we allocate for skeletal mesh work to be performed.\n")
+	TEXT("When overbudget various other CVars come into play, such as a.Budget.AlwaysTickFalloffAggression and a.Budget.InterpolationFalloffAggression."),
+	FConsoleVariableDelegate::CreateLambda([](IConsoleVariable* InVariable)
+	{
+		GBudgetParameters.BudgetInMs = FMath::Max(GBudgetParameters.BudgetInMs, 0.1f);
+		GOnCVarParametersChanged.Broadcast();
+	}),
+	ECVF_Scalability);
+
+static FAutoConsoleVariableRef CVarSkelBatch_MinQuality(
+	TEXT("a.Budget.MinQuality"),
+	GBudgetParameters.MinQuality,
+	TEXT("Values [0.0, 1.0], Default = 0.0\n")
+	TEXT("The minimum quality metric allowed. Quality is determined simply by NumComponentsTickingThisFrame / NumComponentsThatWeNeedToTick.\n")
+	TEXT("If this is anything other than 0.0 then we can potentially go over our time budget."),
+	FConsoleVariableDelegate::CreateLambda([](IConsoleVariable* InVariable)
+	{
+		GBudgetParameters.MinQuality = FMath::Clamp(GBudgetParameters.MinQuality, 0.0f, 1.0f);
+		GOnCVarParametersChanged.Broadcast();
+	}),
+	ECVF_Scalability);
+
+static FAutoConsoleVariableRef CVarSkelBatch_MaxTickRate(
+	TEXT("a.Budget.MaxTickRate"),
+	GBudgetParameters.MaxTickRate,
+	TEXT("Values >= 1, Default = 10\n")
+	TEXT("The maximum tick rate we allow. If this is set then we can potentially go over budget, but keep quality of individual meshes to a reasonable level.\n"),
+	FConsoleVariableDelegate::CreateLambda([](IConsoleVariable* InVariable)
+	{
+		GBudgetParameters.MaxTickRate = FMath::Max(GBudgetParameters.MaxTickRate, 1);
+		GOnCVarParametersChanged.Broadcast();
+	}),
+	ECVF_Scalability);
+
+static FAutoConsoleVariableRef CVarSkelBatch_WorkUnitSmoothingSpeed(
+	TEXT("a.Budget.WorkUnitSmoothingSpeed"),
+	GBudgetParameters.WorkUnitSmoothingSpeed,
+	TEXT("Values > 0.1, Default = 5.0\n")
+	TEXT("The speed at which the average work unit converges on the measured amount.\n"),
+	FConsoleVariableDelegate::CreateLambda([](IConsoleVariable* InVariable)
+	{
+		GBudgetParameters.WorkUnitSmoothingSpeed = FMath::Max(GBudgetParameters.WorkUnitSmoothingSpeed, 0.1f);
+		GOnCVarParametersChanged.Broadcast();
+	}));
+
+static FAutoConsoleVariableRef CVarSkelBatch_AlwaysTickFalloffAggression(
+	TEXT("a.Budget.AlwaysTickFalloffAggression"),
+	GBudgetParameters.AlwaysTickFalloffAggression,
+	TEXT("Range [0.1, 0.9], Default = 0.8\n")
+	TEXT("Controls the rate at which 'always ticked' components falloff under load.\n")
+	TEXT("Higher values mean that we reduce the number of always ticking components by a larger amount when the allocated time budget is exceeded."),
+	FConsoleVariableDelegate::CreateLambda([](IConsoleVariable* InVariable)
+	{
+		GBudgetParameters.AlwaysTickFalloffAggression = FMath::Clamp(GBudgetParameters.AlwaysTickFalloffAggression, 0.1f, 0.9f);
+		GOnCVarParametersChanged.Broadcast();
+	}),
+	ECVF_Scalability);
+
+static FAutoConsoleVariableRef CVarSkelBatch_InterpolationFalloffAggression(
+	TEXT("a.Budget.InterpolationFalloffAggression"),
+	GBudgetParameters.InterpolationFalloffAggression,
+	TEXT("Range [0.1, 0.9], Default = 0.4\n")
+	TEXT("Controls the rate at which interpolated components falloff under load.\n")
+	TEXT("Higher values mean that we reduce the number of interpolated components by a larger amount when the allocated time budget is exceeded.\n")
+	TEXT("Components are only interpolated when the time budget is exceeded."),
+	FConsoleVariableDelegate::CreateLambda([](IConsoleVariable* InVariable)
+	{
+		GBudgetParameters.InterpolationFalloffAggression = FMath::Clamp(GBudgetParameters.InterpolationFalloffAggression, 0.1f, 0.9f);
+		GOnCVarParametersChanged.Broadcast();
+	}),
+	ECVF_Scalability);
+
+static FAutoConsoleVariableRef CVarSkelBatch_InterpolationMaxRate(
+	TEXT("a.Budget.InterpolationMaxRate"),
+	GBudgetParameters.InterpolationMaxRate,
+	TEXT("Values > 1, Default = 6\n")
+	TEXT("Controls the rate at which ticks happen when interpolating.\n"),
+	FConsoleVariableDelegate::CreateLambda([](IConsoleVariable* InVariable)
+	{
+		GBudgetParameters.InterpolationMaxRate = FMath::Max(GBudgetParameters.InterpolationMaxRate, 2);
+		GOnCVarParametersChanged.Broadcast();
+	}),
+	ECVF_Scalability);
+
+static FAutoConsoleVariableRef CVarSkelBatch_MaxInterpolatedComponents(
+	TEXT("a.Budget.MaxInterpolatedComponents"),
+	GBudgetParameters.MaxInterpolatedComponents,
+	TEXT("Range >= 0, Default = 16\n")
+	TEXT("Max number of components to inteprolate before we start throttling.\n"),
+	FConsoleVariableDelegate::CreateLambda([](IConsoleVariable* InVariable)
+	{
+		GBudgetParameters.MaxInterpolatedComponents = FMath::Max(GBudgetParameters.MaxInterpolatedComponents, 0);
+		GOnCVarParametersChanged.Broadcast();
+	}),
+	ECVF_Scalability);
+
+static FAutoConsoleVariableRef CVarSkelBatch_InterpolationTickMultiplier(
+	TEXT("a.Budget.InterpolationTickMultiplier"),
+	GBudgetParameters.InterpolationTickMultiplier,
+	TEXT("Range [0.1, 0.9], Default = 0.75\n")
+	TEXT("Controls the expected value an amortized interpolated tick will take compared to a 'normal' tick.\n"),
+	FConsoleVariableDelegate::CreateLambda([](IConsoleVariable* InVariable)
+	{
+		GBudgetParameters.InterpolationTickMultiplier = FMath::Clamp(GBudgetParameters.InterpolationTickMultiplier, 0.1f, 0.9f);
+		GOnCVarParametersChanged.Broadcast();
+	}),
+	ECVF_Scalability);
+
+static FAutoConsoleVariableRef CVarSkelBatch_InitialEstimatedWorkUnitTime(
+	TEXT("a.Budget.InitialEstimatedWorkUnitTime"),
+	GBudgetParameters.InitialEstimatedWorkUnitTimeMs,
+	TEXT("Values > 0.0, Default = 0.08\n")
+	TEXT("Controls the time in milliseconds we expect, on average, for a skeletal mesh component to execute.\n")
+	TEXT("The value only applies for the first tick of a component, after which we use the real time the tick takes.\n"),
+	FConsoleVariableDelegate::CreateLambda([](IConsoleVariable* InVariable)
+	{
+		GBudgetParameters.InitialEstimatedWorkUnitTimeMs = FMath::Max(GBudgetParameters.InitialEstimatedWorkUnitTimeMs, KINDA_SMALL_NUMBER);
+		GOnCVarParametersChanged.Broadcast();
+	}),
+	ECVF_Scalability);
+
+static FAutoConsoleVariableRef CVarSkelBatch_MaxTickedOffsreenComponents(
+	TEXT("a.Budget.MaxTickedOffsreen"),
+	GBudgetParameters.MaxTickedOffsreenComponents,
+	TEXT("Values >= 1, Default = 4\n")
+	TEXT("The maximum number of offscreen components we tick (most significant first)\n"),
+	FConsoleVariableDelegate::CreateLambda([](IConsoleVariable* InVariable)
+	{
+		GBudgetParameters.MaxTickedOffsreenComponents = FMath::Max(GBudgetParameters.MaxTickedOffsreenComponents, 1);
+		GOnCVarParametersChanged.Broadcast();
+	}),
+	ECVF_Scalability);
+
+static FAutoConsoleVariableRef CVarSkelBatch_StateChangeThrottleInFrames(
+	TEXT("a.Budget.StateChangeThrottleInFrames"),
+	GBudgetParameters.StateChangeThrottleInFrames,
+	TEXT("Range [1, 128], Default = 30\n")
+	TEXT("Prevents throttle values from changing too often due to system and load noise.\n"),
+	FConsoleVariableDelegate::CreateLambda([](IConsoleVariable* InVariable)
+	{
+		GBudgetParameters.StateChangeThrottleInFrames = FMath::Clamp(GBudgetParameters.StateChangeThrottleInFrames, 1, 128);
+		GOnCVarParametersChanged.Broadcast();
+	}),
+	ECVF_Scalability);
+
+
+static FAutoConsoleVariableRef CVarSkelBatch_BudgetFactorBeforeReducedWork(
+	TEXT("a.Budget.BudgetFactorBeforeReducedWork"),
+	GBudgetParameters.BudgetFactorBeforeReducedWork,
+	TEXT("Range > 1, Default = 1.5\n")
+	TEXT("Reduced work will be delayed until budget pressure goes over this amount.\n"),
+	FConsoleVariableDelegate::CreateLambda([](IConsoleVariable* InVariable)
+	{
+		GBudgetParameters.BudgetFactorBeforeReducedWork = FMath::Max(GBudgetParameters.BudgetFactorBeforeReducedWork, 1.0f);
+		GOnCVarParametersChanged.Broadcast();
+	}),
+	ECVF_Scalability);
+
+static FAutoConsoleVariableRef CVarSkelBatch_BudgetFactorBeforeReducedWorkEpsilon(
+	TEXT("a.Budget.BudgetFactorBeforeReducedWorkEpsilon"),
+	GBudgetParameters.BudgetFactorBeforeReducedWorkEpsilon,
+	TEXT("Range > 0.0, Default = 0.25\n")
+	TEXT("Increased work will be delayed until budget pressure goes under a.Budget.BudgetFactorBeforeReducedWork minus this amount.\n"),
+	FConsoleVariableDelegate::CreateLambda([](IConsoleVariable* InVariable)
+	{
+		GBudgetParameters.BudgetFactorBeforeReducedWorkEpsilon = FMath::Max(GBudgetParameters.BudgetFactorBeforeReducedWorkEpsilon, 0.0f);
+		GOnCVarParametersChanged.Broadcast();
+	}),
+	ECVF_Scalability);
+
+static FAutoConsoleVariableRef CVarSkelBatch_BudgetPressureSmoothingSpeed(
+	TEXT("a.Budget.BudgetPressureSmoothingSpeed"),
+	GBudgetParameters.BudgetPressureSmoothingSpeed,
+	TEXT("Range > 0.0, Default = 3.0\n")
+	TEXT("How much to smooth the budget pressure value used to throttle reduced work.\n"),
+	FConsoleVariableDelegate::CreateLambda([](IConsoleVariable* InVariable)
+	{
+		GBudgetParameters.BudgetPressureSmoothingSpeed = FMath::Max(GBudgetParameters.BudgetPressureSmoothingSpeed, KINDA_SMALL_NUMBER);
+		GOnCVarParametersChanged.Broadcast();
+	}),
+	ECVF_Scalability);
+
+static FAutoConsoleVariableRef CVarSkelBatch_ReducedWorkThrottleMinInFrames(
+	TEXT("a.Budget.ReducedWorkThrottleMinInFrames"),
+	GBudgetParameters.ReducedWorkThrottleMinInFrames,
+	TEXT("Range [1, 255], Default = 2\n")
+	TEXT("Prevents reduced work from changing too often due to system and load noise. Min value used when over budget pressure (i.e. aggressive reduction).\n"),
+	FConsoleVariableDelegate::CreateLambda([](IConsoleVariable* InVariable)
+	{
+		GBudgetParameters.ReducedWorkThrottleMinInFrames = FMath::Clamp(GBudgetParameters.ReducedWorkThrottleMinInFrames, 1, 255);
+		GOnCVarParametersChanged.Broadcast();
+	}),
+	ECVF_Scalability);
+
+static FAutoConsoleVariableRef CVarSkelBatch_ReducedWorkThrottleMaxInFrames(
+	TEXT("a.Budget.ReducedWorkThrottleMaxInFrames"),
+	GBudgetParameters.ReducedWorkThrottleMaxInFrames,
+	TEXT("Range [1, 255], Default = 20\n")
+	TEXT("Prevents reduced work from changing too often due to system and load noise. Max value used when under budget pressure.\n"),
+	FConsoleVariableDelegate::CreateLambda([](IConsoleVariable* InVariable)
+	{
+		GBudgetParameters.ReducedWorkThrottleMaxInFrames = FMath::Clamp(GBudgetParameters.ReducedWorkThrottleMaxInFrames, 1, 255);
+		GOnCVarParametersChanged.Broadcast();
+	}),
+	ECVF_Scalability);
+
+static FAutoConsoleVariableRef CVarSkelBatch_BudgetFactorBeforeAggressiveReducedWork(
+	TEXT("a.Budget.BudgetFactorBeforeAggressiveReducedWork"),
+	GBudgetParameters.BudgetFactorBeforeAggressiveReducedWork,
+	TEXT("Range > 1, Default = 2.0\n")
+	TEXT("Reduced work will be applied more rapidly when budget pressure goes over this amount.\n"),
+	FConsoleVariableDelegate::CreateLambda([](IConsoleVariable* InVariable)
+	{
+		GBudgetParameters.BudgetFactorBeforeAggressiveReducedWork = FMath::Max(GBudgetParameters.BudgetFactorBeforeAggressiveReducedWork, 1.0f);
+		GOnCVarParametersChanged.Broadcast();
+	}),
+	ECVF_Scalability);
+
+static FAutoConsoleVariableRef CVarSkelBatch_ReducedWorkThrottleMaxPerFrame(
+	TEXT("a.Budget.ReducedWorkThrottleMaxPerFrame"),
+	GBudgetParameters.ReducedWorkThrottleMaxPerFrame,
+	TEXT("Range [1, 255], Default = 4\n")
+	TEXT("Controls the max number of components that are switched to/from reduced work per tick.\n"),
+	FConsoleVariableDelegate::CreateLambda([](IConsoleVariable* InVariable)
+	{
+		GBudgetParameters.ReducedWorkThrottleMaxPerFrame = FMath::Clamp(GBudgetParameters.ReducedWorkThrottleMaxPerFrame, 1, 255);
+		GOnCVarParametersChanged.Broadcast();
+	}),
+	ECVF_Scalability);
+
+static FAutoConsoleVariableRef CVarSkelBatch_BudgetPressureBeforeEmergencyReducedWork(
+	TEXT("a.Budget.GBudgetPressureBeforeEmergencyReducedWork"),
+	GBudgetParameters.BudgetPressureBeforeEmergencyReducedWork,
+	TEXT("Range > 0.0, Default = 2.5\n")
+	TEXT("Controls the budget pressure where emergency reduced work (applied to all components except those that are bAlwaysTick).\n"),
+	FConsoleVariableDelegate::CreateLambda([](IConsoleVariable* InVariable)
+	{
+		GBudgetParameters.BudgetPressureBeforeEmergencyReducedWork = FMath::Max(GBudgetParameters.BudgetPressureBeforeEmergencyReducedWork, 0.0f);
+		GOnCVarParametersChanged.Broadcast();
+	}),
+	ECVF_Scalability);
\ No newline at end of file
diff --git a/Engine/Plugins/Runtime/AnimationBudgetAllocator/Source/AnimationBudgetAllocator/Private/AnimationBudgetAllocatorCVars.h b/Engine/Plugins/Runtime/AnimationBudgetAllocator/Source/AnimationBudgetAllocator/Private/AnimationBudgetAllocatorCVars.h
new file mode 100644
--- /dev/null
+++ b/Engine/Plugins/Runtime/AnimationBudgetAllocator/Source/AnimationBudgetAllocator/Private/AnimationBudgetAllocatorCVars.h
@@ -0,0 +1,25 @@
+// Copyright 1998-2019 Epic Games, Inc. All Rights Reserved.
+
+#pragma once
+
+#include "CoreMinimal.h"
+#include "DrawDebugHelpers.h"
+
+struct FAnimationBudgetAllocatorParameters;
+
+/** Enabled/disabled flag */
+extern int32 GAnimationBudgetEnabled;
+
+#if ENABLE_DRAW_DEBUG
+/** Debug rendering flag */
+extern int32 GAnimationBudgetDebugEnabled;
+
+/** Controls whether debug rendering shows addresses of component data for debugging */
+extern int32 GAnimationBudgetDebugShowAddresses;
+#endif
+
+/** CVar-driven parameter block */
+extern FAnimationBudgetAllocatorParameters GBudgetParameters;
+
+/** Delegate broadcast when parameter block changes */
+extern FSimpleMulticastDelegate GOnCVarParametersChanged;
\ No newline at end of file
diff --git a/Engine/Plugins/Runtime/AnimationBudgetAllocator/Source/AnimationBudgetAllocator/Private/AnimationBudgetBlueprintLibrary.cpp b/Engine/Plugins/Runtime/AnimationBudgetAllocator/Source/AnimationBudgetAllocator/Private/AnimationBudgetBlueprintLibrary.cpp
--- a/Engine/Plugins/Runtime/AnimationBudgetAllocator/Source/AnimationBudgetAllocator/Private/AnimationBudgetBlueprintLibrary.cpp
+++ b/Engine/Plugins/Runtime/AnimationBudgetAllocator/Source/AnimationBudgetAllocator/Private/AnimationBudgetBlueprintLibrary.cpp
@@ -17,4 +17,18 @@ void UAnimationBudgetBlueprintLibrary::EnableAnimationBudget(UObject* WorldConte
 			AnimationBudgetAllocator->SetEnabled(bEnabled);
 		}
 	}
+}
+
+void UAnimationBudgetBlueprintLibrary::SetAnimationBudgetParameters(UObject* WorldContextObject, const FAnimationBudgetAllocatorParameters& InParameters)
+{
+	if(UWorld* World = GEngine->GetWorldFromContextObject(WorldContextObject, EGetWorldErrorMode::LogAndReturnNull))
+	{
+		FAnimationBudgetAllocatorModule& AnimationBudgetAllocatorModule = FModuleManager::LoadModuleChecked<FAnimationBudgetAllocatorModule>("AnimationBudgetAllocator");
+		IAnimationBudgetAllocator* AnimationBudgetAllocator = AnimationBudgetAllocatorModule.GetBudgetAllocatorForWorld(World);
+		if(AnimationBudgetAllocator)
+		{
+			AnimationBudgetAllocator->SetParameters(InParameters);
+		}
+	}
+}
 }
\ No newline at end of file
diff --git a/Engine/Plugins/Runtime/AnimationBudgetAllocator/Source/AnimationBudgetAllocator/Private/AnimationBudgetBlueprintLibrary.h b/Engine/Plugins/Runtime/AnimationBudgetAllocator/Source/AnimationBudgetAllocator/Private/AnimationBudgetBlueprintLibrary.h
--- a/Engine/Plugins/Runtime/AnimationBudgetAllocator/Source/AnimationBudgetAllocator/Private/AnimationBudgetBlueprintLibrary.h
+++ b/Engine/Plugins/Runtime/AnimationBudgetAllocator/Source/AnimationBudgetAllocator/Private/AnimationBudgetBlueprintLibrary.h
@@ -4,20 +4,28 @@
 
 #include "CoreMinimal.h"
 #include "Kismet/BlueprintFunctionLibrary.h"
+#include "AnimationBudgetAllocatorParameters.h"
 #include "AnimationBudgetBlueprintLibrary.generated.h"
 
 /**
  * Function library to expose the budget allocator to Blueprints
  */
-UCLASS(meta = (ScriptName = "AnimationBudgetLibrary"))
+UCLASS(meta = (ScriptName = "Animation Budget"))
 class UAnimationBudgetBlueprintLibrary : public UBlueprintFunctionLibrary
 {
 	GENERATED_BODY()
 
 	/**
 	 * Enable/disable the animation budgeting system.
-	 * Note that the system can also be disabled 'globally' via CVar, which overrides this setting.
+	 * Note that the system can also be disabled globally via CVar, which overrides this setting.
 	 */
 	UFUNCTION(BlueprintCallable, Category = "Animation Budget", meta=(WorldContext="WorldContextObject"))
 	static void EnableAnimationBudget(UObject* WorldContextObject, bool bEnabled);
-};
\ No newline at end of file
+
+	/**
+	 * Set the parameters of the animation budgeting system.
+	 * Note that the parameters can also be set globally via CVar.
+	 */
+	UFUNCTION(BlueprintCallable, Category = "Animation Budget", meta=(WorldContext="WorldContextObject"))
+	static void SetAnimationBudgetParameters(UObject* WorldContextObject, const FAnimationBudgetAllocatorParameters& InParameters);
+};
diff --git a/Engine/Plugins/Runtime/AnimationBudgetAllocator/Source/AnimationBudgetAllocator/Public/AnimationBudgetAllocatorParameters.h b/Engine/Plugins/Runtime/AnimationBudgetAllocator/Source/AnimationBudgetAllocator/Public/AnimationBudgetAllocatorParameters.h
new file mode 100644
--- /dev/null
+++ b/Engine/Plugins/Runtime/AnimationBudgetAllocator/Source/AnimationBudgetAllocator/Public/AnimationBudgetAllocatorParameters.h
@@ -0,0 +1,158 @@
+// Copyright 1998-2018 Epic Games, Inc. All Rights Reserved.
+
+#include "AnimationBudgetAllocatorParameters.generated.h"
+
+#pragma once
+
+/** Parameter block used to control the behavior of the budget allocator */
+USTRUCT(BlueprintType)
+struct FAnimationBudgetAllocatorParameters
+{
+	GENERATED_BODY()
+
+	/**
+	 * Values > 0.1.
+	 * The time in milliseconds that we allocate for skeletal mesh work to be performed.
+	 * When overbudget various other parameters come into play, such as AlwaysTickFalloffAggression and InterpolationFalloffAggression.
+	 */
+	UPROPERTY(BlueprintReadWrite, Category = "Parameters")
+	float BudgetInMs = 1.0f;
+
+	/**
+	 * Values [0.0, 1.0].
+	 * The minimum quality metric allowed. Quality is determined simply by NumComponentsTickingThisFrame / NumComponentsThatWeNeedToTick.
+	 * If this is anything other than 0.0 then we can potentially go over our time budget.
+	 */
+	UPROPERTY(BlueprintReadWrite, Category = "Parameters")
+	float MinQuality = 0.0f;
+
+	/** 
+	 * Values >= 1.
+	 * The maximum tick rate we allow. If this is set then we can potentially go over budget, but keep quality of individual meshes to a reasonable level. 
+	 */
+	UPROPERTY(BlueprintReadWrite, Category = "Parameters")
+	int32 MaxTickRate = 10;
+
+	/** 
+	 * Values > 0.1.
+	 * The speed at which the average work unit converges on the measured amount. 
+	 */
+	UPROPERTY(BlueprintReadWrite, Category = "Parameters")
+	float WorkUnitSmoothingSpeed = 5.0f;
+
+	/**
+	 * Range [0.1, 0.9].
+	 * Controls the rate at which 'always ticked' components falloff under load.
+	 * Higher values mean that we reduce the number of always ticking components by a larger amount when the allocated time budget is exceeded.
+	*/
+	UPROPERTY(BlueprintReadWrite, Category = "Parameters")
+	float AlwaysTickFalloffAggression = 0.8f;
+
+	/**
+	 * Range [0.1, 0.9].
+	 * Controls the rate at which interpolated components falloff under load.
+	 * Higher values mean that we reduce the number of interpolated components by a larger amount when the allocated time budget is exceeded.
+	 * Components are only interpolated when the time budget is exceeded.
+	 */
+	UPROPERTY(BlueprintReadWrite, Category = "Parameters")
+	float InterpolationFalloffAggression = 0.4f;
+
+	/** 
+	 * Values > 1.
+	 * Controls the rate at which ticks happen when interpolating. 
+	 */
+	UPROPERTY(BlueprintReadWrite, Category = "Parameters")
+	int32 InterpolationMaxRate = 6;
+
+	/** 
+	 * Range >= 0.
+	 * Max number of components to interpolate before we start throttling. 
+	 */
+	UPROPERTY(BlueprintReadWrite, Category = "Parameters")
+	int32 MaxInterpolatedComponents = 16;
+
+	/** 
+	 * Range [0.1, 0.9].
+	 * Controls the expected value an amortized interpolated tick will take compared to a 'normal' tick. 
+	 */
+	UPROPERTY(BlueprintReadWrite, Category = "Parameters")
+	float InterpolationTickMultiplier = 0.75f;
+
+	/** 
+	 * Values > 0.0.
+	 * Controls the time in milliseconds we expect, on average, for a skeletal mesh component to execute.
+	 * The value only applies for the first tick of a component, after which we use the real time the tick takes. 
+	 */
+	UPROPERTY(BlueprintReadWrite, Category = "Parameters")
+	float InitialEstimatedWorkUnitTimeMs = 0.08f;
+
+	/**
+	 * Values >= 1
+	 * The maximum number of offscreen components we tick (most significant first)
+	 */
+	UPROPERTY(BlueprintReadWrite, Category = "Parameters")
+	int32 MaxTickedOffsreenComponents = 4;
+
+	/**
+	 * Range [1, 128]
+	 * Prevents throttle values from changing too often due to system and load noise.
+	 */
+	UPROPERTY(BlueprintReadWrite, Category = "Parameters")
+	int32 StateChangeThrottleInFrames = 30;
+
+	/**
+	 * Range > 1
+	 * Reduced work will be delayed until budget pressure goes over this amount.
+	 */
+	UPROPERTY(BlueprintReadWrite, Category = "Parameters")
+	float BudgetFactorBeforeReducedWork = 1.5f;
+
+	/** 
+	 * Range > 0.0.
+	 * Increased work will be delayed until budget pressure goes under BudgetFactorBeforeReducedWork minus this amount. 
+	 */
+	UPROPERTY(BlueprintReadWrite, Category = "Parameters")
+	float BudgetFactorBeforeReducedWorkEpsilon = 0.25f;
+
+	/**
+	 * Range > 0.0.
+	 * How much to smooth the budget pressure value used to throttle reduced work.
+	 */
+	UPROPERTY(BlueprintReadWrite, Category = "Parameters")
+	float BudgetPressureSmoothingSpeed = 3.0f;
+
+	/**
+	 * Range [1, 255].
+	 * Prevents reduced work from changing too often due to system and load noise. Min value used when over budget pressure (i.e. aggressive reduction).
+	 */
+	UPROPERTY(BlueprintReadWrite, Category = "Parameters")
+	int32 ReducedWorkThrottleMinInFrames = 2;
+
+	/**
+	 * Range [1, 255].
+	 * Prevents reduced work from changing too often due to system and load noise. Max value used when under budget pressure.
+	 */
+	UPROPERTY(BlueprintReadWrite, Category = "Parameters")
+	int32 ReducedWorkThrottleMaxInFrames = 20;
+
+	/**
+	 * Range > 1.
+	 * Reduced work will be applied more rapidly when budget pressure goes over this amount.
+	 */
+	UPROPERTY(BlueprintReadWrite, Category = "Parameters")
+	float BudgetFactorBeforeAggressiveReducedWork = 2.0f;
+
+	/**
+	 * Range [1, 255].
+	 * Controls the max number of components that are switched to/from reduced work per tick.
+	 */
+	UPROPERTY(BlueprintReadWrite, Category = "Parameters")
+	int32 ReducedWorkThrottleMaxPerFrame = 4;
+
+	/**
+	 * Range > 0.0.
+	 * Controls the budget pressure where emergency reduced work (applied to all components except those that are bAlwaysTick).
+	 */
+	UPROPERTY(BlueprintReadWrite, Category = "Parameters")
+	float BudgetPressureBeforeEmergencyReducedWork = 2.5f;
+};
\ No newline at end of file
diff --git a/Engine/Plugins/Runtime/AnimationBudgetAllocator/Source/AnimationBudgetAllocator/Public/IAnimationBudgetAllocator.h b/Engine/Plugins/Runtime/AnimationBudgetAllocator/Source/AnimationBudgetAllocator/Public/IAnimationBudgetAllocator.h
--- a/Engine/Plugins/Runtime/AnimationBudgetAllocator/Source/AnimationBudgetAllocator/Public/IAnimationBudgetAllocator.h
+++ b/Engine/Plugins/Runtime/AnimationBudgetAllocator/Source/AnimationBudgetAllocator/Public/IAnimationBudgetAllocator.h
@@ -4,6 +4,7 @@
 
 class USkeletalMeshComponentBudgeted;
 class UWorld;
+struct FAnimationBudgetAllocatorParameters;
 
 /**
  * Dynamically manages skeletal mesh component tick rates to try to maintain a specified budget.
@@ -66,4 +67,7 @@ public:
 
 	/** Get whether this budget allocator is enabled */
 	virtual bool GetEnabled() const = 0;
+
+	/** Set the various parameters */
+	virtual void SetParameters(const FAnimationBudgetAllocatorParameters& InParameters) = 0;
 };
\ No newline at end of file
