Fixed merge error

diff --git a/Engine/Plugins/Runtime/AnimationBudgetAllocator/Source/AnimationBudgetAllocator/Private/AnimationBudgetBlueprintLibrary.cpp b/Engine/Plugins/Runtime/AnimationBudgetAllocator/Source/AnimationBudgetAllocator/Private/AnimationBudgetBlueprintLibrary.cpp
--- a/Engine/Plugins/Runtime/AnimationBudgetAllocator/Source/AnimationBudgetAllocator/Private/AnimationBudgetBlueprintLibrary.cpp
+++ b/Engine/Plugins/Runtime/AnimationBudgetAllocator/Source/AnimationBudgetAllocator/Private/AnimationBudgetBlueprintLibrary.cpp
@@ -30,5 +30,4 @@ void UAnimationBudgetBlueprintLibrary::SetAnimationBudgetParameters(UObject* Wor
 			AnimationBudgetAllocator->SetParameters(InParameters);
 		}
 	}
-}
 }
\ No newline at end of file
