Remove cloth simulation logic from FAnimationBudgetAllocator::TickEnableHelper(). Some use cases require the mesh component to tick but have cloth disabled.

diff --git a/Engine/Plugins/Runtime/AnimationBudgetAllocator/Source/AnimationBudgetAllocator/Private/AnimationBudgetAllocator.cpp b/Engine/Plugins/Runtime/AnimationBudgetAllocator/Source/AnimationBudgetAllocator/Private/AnimationBudgetAllocator.cpp
--- a/Engine/Plugins/Runtime/AnimationBudgetAllocator/Source/AnimationBudgetAllocator/Private/AnimationBudgetAllocator.cpp
+++ b/Engine/Plugins/Runtime/AnimationBudgetAllocator/Source/AnimationBudgetAllocator/Private/AnimationBudgetAllocator.cpp
@@ -152,24 +152,7 @@ void FAnimationBudgetAllocator::SetComponentSignificance(USkeletalMeshComponentB
 
 void FAnimationBudgetAllocator::TickEnableHelper(USkeletalMeshComponent* InComponent, bool bInEnable)
 {
-	if(bInEnable)
-	{
-		InComponent->PrimaryComponentTick.SetTickFunctionEnable(true);
-		if(InComponent->IsClothingSimulationSuspended())
-		{
-			InComponent->ResumeClothingSimulation();
-			InComponent->ClothBlendWeight = 1.0f;
-		}
-	}
-	else
-	{
-		InComponent->PrimaryComponentTick.SetTickFunctionEnable(false);
-		if(!InComponent->IsClothingSimulationSuspended())
-		{
-			InComponent->SuspendClothingSimulation();
-			InComponent->ClothBlendWeight = 0.0f;
-		}
-	}
+	InComponent->PrimaryComponentTick.SetTickFunctionEnable(bInEnable);
 }
 
 void FAnimationBudgetAllocator::QueueSortedComponentIndices(float InDeltaSeconds)
