Animation budget allocator: Skip offscreen components with a significance of zero

Also adds debug rendering for offscreen components

diff --git a/Engine/Plugins/Runtime/AnimationBudgetAllocator/Source/AnimationBudgetAllocator/Private/AnimationBudgetAllocator.cpp b/Engine/Plugins/Runtime/AnimationBudgetAllocator/Source/AnimationBudgetAllocator/Private/AnimationBudgetAllocator.cpp
--- a/Engine/Plugins/Runtime/AnimationBudgetAllocator/Source/AnimationBudgetAllocator/Private/AnimationBudgetAllocator.cpp
+++ b/Engine/Plugins/Runtime/AnimationBudgetAllocator/Source/AnimationBudgetAllocator/Private/AnimationBudgetAllocator.cpp
@@ -244,8 +244,9 @@ void FAnimationBudgetAllocator::QueueSortedComponentIndices(float InDeltaSeconds
 
 				if(bShouldTick)
 				{
-					// Push into a separate limited list if we are 'tick even if not rendered'
-					if(Component->LastRenderTime <= WorldTime && ComponentData.bTickEvenIfNotRendered)
+					// Push into a separate limited list if we are 'tick even if not rendered'.
+					// Skip offscreen components with a significance of zero or less.
+					if(Component->LastRenderTime <= WorldTime && ComponentData.bTickEvenIfNotRendered && ComponentData.Significance > 0.0f)
 					{
 						NonRenderedComponentData.Add(ComponentIndex);
 					}
@@ -707,7 +708,7 @@ void FAnimationBudgetAllocator::Update(float DeltaSeconds)
 				for(int32 ComponentDataIndex : ActorIndicesPair.Value)
 				{
 					FComponentData& ComponentData = AllComponentData[ComponentDataIndex];
-					if(ComponentData.bTickEnabled)
+					if(ComponentData.bTickEnabled && ComponentData.bOnScreen)
 					{
 						if(GAnimationBudgetDebugShowAddresses != 0)
 						{
@@ -715,7 +716,7 @@ void FAnimationBudgetAllocator::Update(float DeltaSeconds)
 						}
 						else
 						{
-							DebugString += FString::Printf(TEXT("%d %s %s\n"), ComponentData.TickRate, ComponentData.bInterpolate ? TEXT("I") : TEXT(" "), ComponentData.bReducedWork ? TEXT("Lo") : TEXT("Hi"));
+							DebugString += FString::Printf(TEXT("%.03f %d %s %s\n"), ComponentData.Significance, ComponentData.TickRate, ComponentData.bInterpolate ? TEXT("I") : TEXT(" "), ComponentData.bReducedWork ? TEXT("Lo") : TEXT("Hi"));
 						}
 					}
 				}
@@ -784,6 +785,49 @@ void FAnimationBudgetAllocator::OnHUDPostRender(AHUD* HUD, UCanvas* Canvas)
 			Graph->DrawExtremesOnGraph(false);
 			Graph->bVisible = true;
 			Graph->Draw(Canvas);
+
+			{
+				TMap<AActor*, TArray<int32>> ActorMap;
+				for (int32 ComponentDataIndex : AllSortedComponentData)
+				{
+					FComponentData& ComponentData = AllComponentData[ComponentDataIndex];
+					TArray<int32>& ComponentIndexArray = ActorMap.FindOrAdd(ComponentData.Component->GetOwner());
+					ComponentIndexArray.Add(ComponentDataIndex);
+				}
+
+				float LineOffset = 0.0f;
+
+				for(const TPair<AActor*, TArray<int32>>& ActorIndicesPair : ActorMap)
+				{
+					FVector Location = ActorIndicesPair.Key->GetActorLocation();
+
+					FString DebugString;
+
+					float SubLineOffset = 0.0f;
+
+					for(int32 ComponentDataIndex : ActorIndicesPair.Value)
+					{
+						FComponentData& ComponentData = AllComponentData[ComponentDataIndex];
+						if(ComponentData.bTickEnabled && !ComponentData.bOnScreen)
+						{
+							if(GAnimationBudgetDebugShowAddresses != 0)
+							{
+								DebugString += FString::Printf(TEXT("0x%llx %d %s %s\n"), &ComponentData, ComponentData.TickRate, ComponentData.bInterpolate ? TEXT("I") : TEXT(" "), ComponentData.bReducedWork ? TEXT("Lo") : TEXT("Hi"));
+							}
+							else
+							{
+								DebugString += FString::Printf(TEXT("%.03f %d %s %s\n"), ComponentData.Significance, ComponentData.TickRate, ComponentData.bInterpolate ? TEXT("I") : TEXT(" "), ComponentData.bReducedWork ? TEXT("Lo") : TEXT("Hi"));
+							}
+
+							SubLineOffset += 15.0f;
+						}
+					}
+
+					Canvas->DrawText(GEngine->GetSmallFont(), FText::FromString(DebugString), 100.0f, 100.0f + LineOffset);
+
+					LineOffset += SubLineOffset;
+				}	
+			}
 		}
 	}
 }
