Disabled BudgetedAnimation CSV stat in shipping builds

diff --git a/Engine/Plugins/Runtime/AnimationBudgetAllocator/Source/AnimationBudgetAllocator/Private/SkeletalMeshComponentBudgeted.cpp b/Engine/Plugins/Runtime/AnimationBudgetAllocator/Source/AnimationBudgetAllocator/Private/SkeletalMeshComponentBudgeted.cpp
--- a/Engine/Plugins/Runtime/AnimationBudgetAllocator/Source/AnimationBudgetAllocator/Private/SkeletalMeshComponentBudgeted.cpp
+++ b/Engine/Plugins/Runtime/AnimationBudgetAllocator/Source/AnimationBudgetAllocator/Private/SkeletalMeshComponentBudgeted.cpp
@@ -43,7 +43,9 @@ void USkeletalMeshComponentBudgeted::EndPlay(const EEndPlayReason::Type EndPlayR
 
 void USkeletalMeshComponentBudgeted::TickComponent(float DeltaTime, enum ELevelTick TickType, FActorComponentTickFunction *ThisTickFunction)
 {
+#if !UE_BUILD_SHIPPING
 	CSV_SCOPED_TIMING_STAT(AnimationBudget, BudgetedAnimation);
+#endif
 
 	if(AnimationBudgetAllocator)
 	{
@@ -61,7 +63,9 @@ void USkeletalMeshComponentBudgeted::TickComponent(float DeltaTime, enum ELevelT
 
 void USkeletalMeshComponentBudgeted::CompleteParallelAnimationEvaluation(bool bDoPostAnimEvaluation)
 {
+#if !UE_BUILD_SHIPPING
 	CSV_SCOPED_TIMING_STAT(AnimationBudget, BudgetedAnimation);
+#endif
 
 	if(AnimationBudgetAllocator)
 	{
