Allow game code to call SetComponentTickEnabled and SetComponentSignificance directly on the budgeted component instead of through a reference to the animation budget allocator

diff --git a/Engine/Plugins/Runtime/AnimationBudgetAllocator/Source/AnimationBudgetAllocator/Private/SkeletalMeshComponentBudgeted.cpp b/Engine/Plugins/Runtime/AnimationBudgetAllocator/Source/AnimationBudgetAllocator/Private/SkeletalMeshComponentBudgeted.cpp
--- a/Engine/Plugins/Runtime/AnimationBudgetAllocator/Source/AnimationBudgetAllocator/Private/SkeletalMeshComponentBudgeted.cpp
+++ b/Engine/Plugins/Runtime/AnimationBudgetAllocator/Source/AnimationBudgetAllocator/Private/SkeletalMeshComponentBudgeted.cpp
@@ -41,6 +41,26 @@ void USkeletalMeshComponentBudgeted::EndPlay(const EEndPlayReason::Type EndPlayR
 	Super::EndPlay(EndPlayReason);
 }
 
+void USkeletalMeshComponentBudgeted::SetComponentTickEnabled(bool bEnabled)
+{
+	if (AnimationBudgetAllocator)
+	{
+		AnimationBudgetAllocator->SetComponentTickEnabled(this, bEnabled);
+	}
+	else
+	{
+		Super::SetComponentTickEnabled(bEnabled);
+	}
+}
+
+void USkeletalMeshComponentBudgeted::SetComponentSignificance(float Significance, bool bNeverSkip, bool bTickEvenIfNotRendered, bool bAllowReducedWork, bool bForceInterpolate)
+{
+	if (AnimationBudgetAllocator)
+	{
+		AnimationBudgetAllocator->SetComponentSignificance(this, Significance, bNeverSkip, bTickEvenIfNotRendered, bAllowReducedWork, bForceInterpolate);
+	}
+}
+
 void USkeletalMeshComponentBudgeted::TickComponent(float DeltaTime, enum ELevelTick TickType, FActorComponentTickFunction *ThisTickFunction)
 {
 #if !UE_BUILD_SHIPPING
diff --git a/Engine/Plugins/Runtime/AnimationBudgetAllocator/Source/AnimationBudgetAllocator/Public/SkeletalMeshComponentBudgeted.h b/Engine/Plugins/Runtime/AnimationBudgetAllocator/Source/AnimationBudgetAllocator/Public/SkeletalMeshComponentBudgeted.h
--- a/Engine/Plugins/Runtime/AnimationBudgetAllocator/Source/AnimationBudgetAllocator/Public/SkeletalMeshComponentBudgeted.h
+++ b/Engine/Plugins/Runtime/AnimationBudgetAllocator/Source/AnimationBudgetAllocator/Public/SkeletalMeshComponentBudgeted.h
@@ -25,6 +25,12 @@ class ANIMATIONBUDGETALLOCATOR_API USkeletalMeshComponentBudgeted : public USkel
 public:
 	USkeletalMeshComponentBudgeted(const FObjectInitializer& ObjectInitializer);
 
+	// UActorComponent interface
+	virtual void SetComponentTickEnabled(bool bEnabled) override;
+
+	/** Updates significance budget if this component has been registered with a AnimationBudgetAllocator */
+	void SetComponentSignificance(float Significance, bool bNeverSkip = false, bool bTickEvenIfNotRendered = false, bool bAllowReducedWork = true, bool bForceInterpolate = false);
+
 	/** Set this component to automatically register with the budget allocator */
 	UFUNCTION(BlueprintSetter)
 	void SetAutoRegisterWithBudgetAllocator(bool bInAutoRegisterWithBudgetAllocator) { bAutoRegisterWithBudgetAllocator = bInAutoRegisterWithBudgetAllocator; }
