Added log warning when calling SetComponentSignifiance() on a budgeted component that doesn't have a budget allocator (ie. unregistered)

#rnx

diff --git a/Engine/Plugins/Runtime/AnimationBudgetAllocator/Source/AnimationBudgetAllocator/Private/SkeletalMeshComponentBudgeted.cpp b/Engine/Plugins/Runtime/AnimationBudgetAllocator/Source/AnimationBudgetAllocator/Private/SkeletalMeshComponentBudgeted.cpp
--- a/Engine/Plugins/Runtime/AnimationBudgetAllocator/Source/AnimationBudgetAllocator/Private/SkeletalMeshComponentBudgeted.cpp
+++ b/Engine/Plugins/Runtime/AnimationBudgetAllocator/Source/AnimationBudgetAllocator/Private/SkeletalMeshComponentBudgeted.cpp
@@ -59,6 +59,10 @@ void USkeletalMeshComponentBudgeted::SetComponentSignificance(float Significance
 	{
 		AnimationBudgetAllocator->SetComponentSignificance(this, Significance, bNeverSkip, bTickEvenIfNotRendered, bAllowReducedWork, bForceInterpolate);
 	}
+	else
+	{
+		UE_LOG(LogSkeletalMesh, Warning, TEXT("SetComponentSignificance called on [%s] before registering with budget allocator"), *GetName());
+	}
 }
 
 void USkeletalMeshComponentBudgeted::TickComponent(float DeltaTime, enum ELevelTick TickType, FActorComponentTickFunction *ThisTickFunction)
