Added animation budget library to allow for testing

diff --git a/Engine/Plugins/Runtime/AnimationBudgetAllocator/Source/AnimationBudgetAllocator/Private/AnimationBudgetBlueprintLibrary.cpp b/Engine/Plugins/Runtime/AnimationBudgetAllocator/Source/AnimationBudgetAllocator/Private/AnimationBudgetBlueprintLibrary.cpp
new file mode 100644
--- /dev/null
+++ b/Engine/Plugins/Runtime/AnimationBudgetAllocator/Source/AnimationBudgetAllocator/Private/AnimationBudgetBlueprintLibrary.cpp
@@ -0,0 +1,20 @@
+// Copyright 1998-2019 Epic Games, Inc. All Rights Reserved.
+
+#include "AnimationBudgetBlueprintLibrary.h"
+#include "AnimationBudgetAllocatorModule.h"
+#include "Modules/ModuleManager.h"
+#include "IAnimationBudgetAllocator.h"
+#include "Engine/Engine.h"
+
+void UAnimationBudgetBlueprintLibrary::EnableAnimationBudget(UObject* WorldContextObject, bool bEnabled)
+{
+	if(UWorld* World = GEngine->GetWorldFromContextObject(WorldContextObject, EGetWorldErrorMode::LogAndReturnNull))
+	{
+		FAnimationBudgetAllocatorModule& AnimationBudgetAllocatorModule = FModuleManager::LoadModuleChecked<FAnimationBudgetAllocatorModule>("AnimationBudgetAllocator");
+		IAnimationBudgetAllocator* AnimationBudgetAllocator = AnimationBudgetAllocatorModule.GetBudgetAllocatorForWorld(World);
+		if(AnimationBudgetAllocator)
+		{
+			AnimationBudgetAllocator->SetEnabled(bEnabled);
+		}
+	}
+}
\ No newline at end of file
diff --git a/Engine/Plugins/Runtime/AnimationBudgetAllocator/Source/AnimationBudgetAllocator/Private/AnimationBudgetBlueprintLibrary.h b/Engine/Plugins/Runtime/AnimationBudgetAllocator/Source/AnimationBudgetAllocator/Private/AnimationBudgetBlueprintLibrary.h
new file mode 100644
--- /dev/null
+++ b/Engine/Plugins/Runtime/AnimationBudgetAllocator/Source/AnimationBudgetAllocator/Private/AnimationBudgetBlueprintLibrary.h
@@ -0,0 +1,23 @@
+// Copyright 1998-2019 Epic Games, Inc. All Rights Reserved.
+
+#pragma once
+
+#include "CoreMinimal.h"
+#include "Kismet/BlueprintFunctionLibrary.h"
+#include "AnimationBudgetBlueprintLibrary.generated.h"
+
+/**
+ * Function library to expose the budget allocator to Blueprints
+ */
+UCLASS(meta = (ScriptName = "Animation Budget"))
+class UAnimationBudgetBlueprintLibrary : public UBlueprintFunctionLibrary
+{
+	GENERATED_BODY()
+
+	/**
+	 * Enable/disable the animation budgeting system.
+	 * Note that the system can also be disabled 'globally' via CVar, which overrides this setting.
+	 */
+	UFUNCTION(BlueprintCallable, Category = "Animation Budget", meta=(WorldContext="WorldContextObject"))
+	static void EnableAnimationBudget(UObject* WorldContextObject, bool bEnabled);
+};
\ No newline at end of file
