Added some extra tuning to the budget allocator to deal better with high budget pressure situations

We now dynamically scale the rate at which components transition to reduced work based on budget pressure.

Also:
Offscreen components now revert to reduced work, so averaging calculations arent impacted.
Offscreen components now allow update rate to switch on first tick. Previously these were throttled to last appearance.

diff --git a/Engine/Plugins/Runtime/AnimationBudgetAllocator/Source/AnimationBudgetAllocator/Private/AnimationBudgetAllocator.cpp b/Engine/Plugins/Runtime/AnimationBudgetAllocator/Source/AnimationBudgetAllocator/Private/AnimationBudgetAllocator.cpp
--- a/Engine/Plugins/Runtime/AnimationBudgetAllocator/Source/AnimationBudgetAllocator/Private/AnimationBudgetAllocator.cpp
+++ b/Engine/Plugins/Runtime/AnimationBudgetAllocator/Source/AnimationBudgetAllocator/Private/AnimationBudgetAllocator.cpp
@@ -164,7 +164,7 @@ static float GInitialEstimatedWorkUnitTimeMs = 0.08f;
 static FAutoConsoleVariableRef CVarSkelBatch_InitialEstimatedWorkUnitTime(
 	TEXT("a.Budget.InitialEstimatedWorkUnitTime"),
 	GInitialEstimatedWorkUnitTimeMs,
-	TEXT("Values > 0.0, Default = 0.15\n")
+	TEXT("Values > 0.0, Default = 0.08\n")
 	TEXT("Controls the time in milliseconds we expect, on average, for a skeletal mesh component to execute.\n")
 	TEXT("The value only applies for the first tick of a component, after which we use the real time the tick takes.\n"),
 	FConsoleVariableDelegate::CreateLambda([](IConsoleVariable* InVariable)
@@ -225,12 +225,12 @@ static FAutoConsoleVariableRef CVarSkelBatch_BudgetFactorBeforeReducedWorkEpsilo
 	}),
 	ECVF_Scalability);
 
-static float GBudgetPressureSmoothingSpeed = 1.5f;
+static float GBudgetPressureSmoothingSpeed = 3.0f;
 
 static FAutoConsoleVariableRef CVarSkelBatch_BudgetPressureSmoothingSpeed(
 	TEXT("a.Budget.BudgetPressureSmoothingSpeed"),
 	GBudgetPressureSmoothingSpeed,
-	TEXT("Range > 0.0, Default = 1.5\n")
+	TEXT("Range > 0.0, Default = 3.0\n")
 	TEXT("How much to smooth the budget pressure value used to throttle reduced work.\n"),
 	FConsoleVariableDelegate::CreateLambda([](IConsoleVariable* InVariable)
 	{
@@ -238,16 +238,55 @@ static FAutoConsoleVariableRef CVarSkelBatch_BudgetPressureSmoothingSpeed(
 	}),
 	ECVF_Scalability);
 
-static int32 GReducedWorkThrottleInFrames = 20;
+static int32 GReducedWorkThrottleMinInFrames = 2;
 
-static FAutoConsoleVariableRef CVarSkelBatch_ReducedWorkThrottleInFrames(
-	TEXT("a.Budget.ReducedWorkThrottleInFrames"),
+static FAutoConsoleVariableRef CVarSkelBatch_ReducedWorkThrottleMinInFrames(
+	TEXT("a.Budget.ReducedWorkThrottleMinInFrames"),
+	GStateChangeThrottleInFrames,
+	TEXT("Range [1, 255], Default = 2\n")
+	TEXT("Prevents reduced work from changing too often due to system and load noise. Min value used when over budget pressure (i.e. aggressive reduction).\n"),
+	FConsoleVariableDelegate::CreateLambda([](IConsoleVariable* InVariable)
+	{
+		GReducedWorkThrottleMinInFrames = FMath::Clamp(GReducedWorkThrottleMinInFrames, 1, 255);
+	}),
+	ECVF_Scalability);
+
+static int32 GReducedWorkThrottleMaxInFrames = 20;
+
+static FAutoConsoleVariableRef CVarSkelBatch_ReducedWorkThrottleMaxInFrames(
+	TEXT("a.Budget.ReducedWorkThrottleMaxInFrames"),
 	GStateChangeThrottleInFrames,
 	TEXT("Range [1, 255], Default = 20\n")
-	TEXT("Prevents reduced work from changing too often due to system and load noise.\n"),
+	TEXT("Prevents reduced work from changing too often due to system and load noise. Max value used when under budget pressure.\n"),
+	FConsoleVariableDelegate::CreateLambda([](IConsoleVariable* InVariable)
+	{
+		GReducedWorkThrottleMaxInFrames = FMath::Clamp(GReducedWorkThrottleMaxInFrames, 1, 255);
+	}),
+	ECVF_Scalability);
+
+static float GBudgetFactorBeforeAggressiveReducedWork = 2.0f;
+
+static FAutoConsoleVariableRef CVarSkelBatch_BudgetFactorBeforeAggressiveReducedWork(
+	TEXT("a.Budget.BudgetFactorBeforeAggressiveReducedWork"),
+	GBudgetFactorBeforeReducedWork,
+	TEXT("Range > 1, Default = 2.0\n")
+	TEXT("Reduced work will be applied more rapidly when budget pressure goes over this amount.\n"),
+	FConsoleVariableDelegate::CreateLambda([](IConsoleVariable* InVariable)
+	{
+		GBudgetFactorBeforeAggressiveReducedWork = FMath::Max(GBudgetFactorBeforeAggressiveReducedWork, 1.0f);
+	}),
+	ECVF_Scalability);
+
+static int32 GReducedWorkThrottleMaxPerFrame = 4;
+
+static FAutoConsoleVariableRef CVarSkelBatch_ReducedWorkThrottleMaxPerFrame(
+	TEXT("a.Budget.ReducedWorkThrottleMaxPerFrame"),
+	GStateChangeThrottleInFrames,
+	TEXT("Range [1, 255], Default = 4\n")
+	TEXT("Controls the max number of components that are switched to/from reduced work per tick.\n"),
 	FConsoleVariableDelegate::CreateLambda([](IConsoleVariable* InVariable)
 	{
-		GReducedWorkThrottleInFrames = FMath::Clamp(GReducedWorkThrottleInFrames, 1, 255);
+		GReducedWorkThrottleMaxPerFrame = FMath::Clamp(GReducedWorkThrottleMaxPerFrame, 1, 255);
 	}),
 	ECVF_Scalability);
 
@@ -498,6 +537,7 @@ void FAnimationBudgetAllocator::QueueSortedComponentIndices(float InDeltaSeconds
 				ReducedWorkComponentData.Add(ComponentIndex);
 			}
 		}
+
 		ComponentIndex++;
 	}
 
@@ -516,17 +556,43 @@ void FAnimationBudgetAllocator::QueueSortedComponentIndices(float InDeltaSeconds
 	const int32 MaxOffscreenComponents = FMath::Min(NonRenderedComponentData.Num(), GMaxTickedOffsreenComponents);
 	if(MaxOffscreenComponents > 0)
 	{
+		auto ReduceWorkForOffscreenComponent = [](FComponentData& InComponentData)
+		{
+			if(InComponentData.bAllowReducedWork && !InComponentData.bReducedWork && InComponentData.Component->OnReduceWork().IsBound())
+			{
+#if WITH_TICK_DEBUG
+				UE_LOG(LogTemp, Warning, TEXT("Force-decreasing offscreen component work (mesh %s) (actor %llx)"), InComponentData.Component->SkeletalMesh ? *InComponentData.Component->SkeletalMesh->GetName() : TEXT("null"), (uint64)InComponentData.Component->GetOwner());
+#endif
+				InComponentData.Component->OnReduceWork().Execute(InComponentData.Component, true);
+				InComponentData.bReducedWork = true;
+			}
+		};
+
 		// Queue first N offscreen ticks
 		int32 NonRenderedComponentIndex = 0;
 		for(; NonRenderedComponentIndex < MaxOffscreenComponents; ++NonRenderedComponentIndex)
 		{
-			QueueComponentTick(AllComponentData[NonRenderedComponentData[NonRenderedComponentIndex]], NonRenderedComponentData[NonRenderedComponentIndex], false);
+			FComponentData& ComponentData = AllComponentData[NonRenderedComponentData[NonRenderedComponentIndex]];
+			QueueComponentTick(ComponentData, NonRenderedComponentData[NonRenderedComponentIndex], false);
+
+			// Always move to reduced work offscreen
+			ReduceWorkForOffscreenComponent(ComponentData);
+
+			// Offscreen will need state changing ASAP when back onscreen
+			ComponentData.StateChangeThrottle = GStateChangeThrottleInFrames + 1;
 		}
 
 		// Disable ticks for the rest
 		for(; NonRenderedComponentIndex < NonRenderedComponentData.Num(); ++NonRenderedComponentIndex)
 		{
+			FComponentData& ComponentData = AllComponentData[NonRenderedComponentData[NonRenderedComponentIndex]];
 			DisableComponentTick(AllComponentData[NonRenderedComponentData[NonRenderedComponentIndex]]);
+
+			// Always move to reduced work offscreen
+			ReduceWorkForOffscreenComponent(ComponentData);
+
+			// Offscreen will need state changing ASAP when back onscreen
+			ComponentData.StateChangeThrottle = GStateChangeThrottleInFrames + 1;
 		}
 
 		// re-sort now we have inserted offscreen components
@@ -737,6 +803,12 @@ int32 FAnimationBudgetAllocator::CalculateWorkDistributionAndQueue(float InDelta
 
 		if(--ReducedComponentWorkCounter <= 0)
 		{
+			float BudgetPressureInterpAlpha = FMath::Clamp((SmoothedBudgetPressure - GBudgetFactorBeforeAggressiveReducedWork) * 0.5f, 0.0f, 1.0f);
+
+			// Scale num components to switch based on budget pressure
+			const int32 NumComponentsToSwitch = (int32)FMath::Lerp(1.0f, (float)GReducedWorkThrottleMaxPerFrame, BudgetPressureInterpAlpha);
+			int32 ComponentsSwitched = 0;
+
 			// If we have any components running reduced work when we have an excess, then move them out of the 'reduced' pool per tick
 			if (ReducedWorkComponentData.Num() > 0 && SmoothedBudgetPressure < GBudgetFactorBeforeReducedWork - GBudgetFactorBeforeReducedWorkEpsilon)
 			{
@@ -750,7 +822,12 @@ int32 FAnimationBudgetAllocator::CalculateWorkDistributionAndQueue(float InDelta
 #endif
 						ComponentData.Component->OnReduceWork().Execute(ComponentData.Component, false);
 						ComponentData.bReducedWork = false;
-						break;
+
+						ComponentsSwitched++;
+						if(ComponentsSwitched >= NumComponentsToSwitch)
+						{
+							break;	
+						}
 					}
 				}
 			}
@@ -767,12 +844,18 @@ int32 FAnimationBudgetAllocator::CalculateWorkDistributionAndQueue(float InDelta
 #endif
 						ComponentData.Component->OnReduceWork().Execute(ComponentData.Component, true);
 						ComponentData.bReducedWork = true;
-						break;
+
+						ComponentsSwitched++;
+						if(ComponentsSwitched >= NumComponentsToSwitch)
+						{
+							break;	
+						}
 					}
 				}
 			}
 
-			ReducedComponentWorkCounter = GReducedWorkThrottleInFrames;
+			// Scale the rate at which we consider reducing component work based on budget pressure
+			ReducedComponentWorkCounter = (int32)FMath::Lerp((float)GReducedWorkThrottleMaxInFrames, (float)GReducedWorkThrottleMinInFrames, BudgetPressureInterpAlpha);
 		}
 	}
 
