More tweaks to budget allocator to avoid initial overbudget peaks

Now we scale tick rate change throttling based on budget pressure too
Also fixes some typos in CVar setup.

diff --git a/Engine/Plugins/Runtime/AnimationBudgetAllocator/Source/AnimationBudgetAllocator/Private/AnimationBudgetAllocator.cpp b/Engine/Plugins/Runtime/AnimationBudgetAllocator/Source/AnimationBudgetAllocator/Private/AnimationBudgetAllocator.cpp
--- a/Engine/Plugins/Runtime/AnimationBudgetAllocator/Source/AnimationBudgetAllocator/Private/AnimationBudgetAllocator.cpp
+++ b/Engine/Plugins/Runtime/AnimationBudgetAllocator/Source/AnimationBudgetAllocator/Private/AnimationBudgetAllocator.cpp
@@ -242,7 +242,7 @@ static int32 GReducedWorkThrottleMinInFrames = 2;
 
 static FAutoConsoleVariableRef CVarSkelBatch_ReducedWorkThrottleMinInFrames(
 	TEXT("a.Budget.ReducedWorkThrottleMinInFrames"),
-	GStateChangeThrottleInFrames,
+	GReducedWorkThrottleMinInFrames,
 	TEXT("Range [1, 255], Default = 2\n")
 	TEXT("Prevents reduced work from changing too often due to system and load noise. Min value used when over budget pressure (i.e. aggressive reduction).\n"),
 	FConsoleVariableDelegate::CreateLambda([](IConsoleVariable* InVariable)
@@ -255,7 +255,7 @@ static int32 GReducedWorkThrottleMaxInFrames = 20;
 
 static FAutoConsoleVariableRef CVarSkelBatch_ReducedWorkThrottleMaxInFrames(
 	TEXT("a.Budget.ReducedWorkThrottleMaxInFrames"),
-	GStateChangeThrottleInFrames,
+	GReducedWorkThrottleMaxInFrames,
 	TEXT("Range [1, 255], Default = 20\n")
 	TEXT("Prevents reduced work from changing too often due to system and load noise. Max value used when under budget pressure.\n"),
 	FConsoleVariableDelegate::CreateLambda([](IConsoleVariable* InVariable)
@@ -268,7 +268,7 @@ static float GBudgetFactorBeforeAggressiveReducedWork = 2.0f;
 
 static FAutoConsoleVariableRef CVarSkelBatch_BudgetFactorBeforeAggressiveReducedWork(
 	TEXT("a.Budget.BudgetFactorBeforeAggressiveReducedWork"),
-	GBudgetFactorBeforeReducedWork,
+	GBudgetFactorBeforeAggressiveReducedWork,
 	TEXT("Range > 1, Default = 2.0\n")
 	TEXT("Reduced work will be applied more rapidly when budget pressure goes over this amount.\n"),
 	FConsoleVariableDelegate::CreateLambda([](IConsoleVariable* InVariable)
@@ -281,7 +281,7 @@ static int32 GReducedWorkThrottleMaxPerFrame = 4;
 
 static FAutoConsoleVariableRef CVarSkelBatch_ReducedWorkThrottleMaxPerFrame(
 	TEXT("a.Budget.ReducedWorkThrottleMaxPerFrame"),
-	GStateChangeThrottleInFrames,
+	GReducedWorkThrottleMaxPerFrame,
 	TEXT("Range [1, 255], Default = 4\n")
 	TEXT("Controls the max number of components that are switched to/from reduced work per tick.\n"),
 	FConsoleVariableDelegate::CreateLambda([](IConsoleVariable* InVariable)
@@ -301,7 +301,7 @@ FComponentData::FComponentData(USkeletalMeshComponentBudgeted* InComponent)
 	, DesiredTickRate(1)
 	, TickRate(1)
 	, SkippedTicks(0)
-	, StateChangeThrottle(0)
+	, StateChangeThrottle(GStateChangeThrottleInFrames)
 	, bTickEnabled(true)
 	, bAlwaysTick(false)
 	, bTickEvenIfNotRendered(false)
@@ -434,7 +434,7 @@ void FAnimationBudgetAllocator::QueueSortedComponentIndices(float InDeltaSeconds
 	{
 		InComponentData.AccumulatedDeltaTime += InDeltaSeconds;
 		InComponentData.bOnScreen = bInOnScreen;
-		InComponentData.StateChangeThrottle++;
+		InComponentData.StateChangeThrottle = InComponentData.StateChangeThrottle < 0 ? InComponentData.StateChangeThrottle : InComponentData.StateChangeThrottle - 1;
 
 		if(InComponentData.bAlwaysTick)
 		{
@@ -579,7 +579,7 @@ void FAnimationBudgetAllocator::QueueSortedComponentIndices(float InDeltaSeconds
 			ReduceWorkForOffscreenComponent(ComponentData);
 
 			// Offscreen will need state changing ASAP when back onscreen
-			ComponentData.StateChangeThrottle = GStateChangeThrottleInFrames + 1;
+			ComponentData.StateChangeThrottle = -1;
 		}
 
 		// Disable ticks for the rest
@@ -592,7 +592,7 @@ void FAnimationBudgetAllocator::QueueSortedComponentIndices(float InDeltaSeconds
 			ReduceWorkForOffscreenComponent(ComponentData);
 
 			// Offscreen will need state changing ASAP when back onscreen
-			ComponentData.StateChangeThrottle = GStateChangeThrottleInFrames + 1;
+			ComponentData.StateChangeThrottle = -1;
 		}
 
 		// re-sort now we have inserted offscreen components
@@ -613,7 +613,7 @@ int32 FAnimationBudgetAllocator::CalculateWorkDistributionAndQueue(float InDelta
 {
 	int32 NumTicked = 0;
 
-	auto QueueForTick = [&NumTicked, this](FComponentData& InComponentData)
+	auto QueueForTick = [&NumTicked, this](FComponentData& InComponentData, int32 InStateChangeThrottleInFrames)
 	{
 		const int32 PrerequisiteHandle = InComponentData.RootPrerequisite != nullptr ? InComponentData.RootPrerequisite->GetAnimationBudgetHandle() : INDEX_NONE;
 		const FComponentData& ComponentDataToCheck = PrerequisiteHandle != INDEX_NONE ? AllComponentData[PrerequisiteHandle] : InComponentData;
@@ -645,10 +645,10 @@ int32 FAnimationBudgetAllocator::CalculateWorkDistributionAndQueue(float InDelta
 			TickEnableHelper(InComponentData.Component, true);
 
 			// Only switch to desired tick rate when we actually tick (throttled)
-			if(bTickThisFrame && InComponentData.StateChangeThrottle > GStateChangeThrottleInFrames)
+			if(bTickThisFrame && InComponentData.StateChangeThrottle < 0)
 			{
 				InComponentData.TickRate = InComponentData.DesiredTickRate;
-				InComponentData.StateChangeThrottle = 0;
+				InComponentData.StateChangeThrottle = InStateChangeThrottleInFrames;
 			}
 
 			NumTicked++;
@@ -761,6 +761,14 @@ int32 FAnimationBudgetAllocator::CalculateWorkDistributionAndQueue(float InDelta
 		BUDGET_CSV_STAT(AnimationBudget, NumInterpolated, NumInterpolated, ECsvCustomStatOp::Set);
 		BUDGET_CSV_STAT(AnimationBudget, NumThrottled, RemainingWorkUnitsToRun, ECsvCustomStatOp::Set);
 
+		const float BudgetPressure = (float)TotalIdealWorkUnits / WorkUnitBudget;
+		SmoothedBudgetPressure = FMath::FInterpTo(SmoothedBudgetPressure, BudgetPressure, InDeltaSeconds, GBudgetPressureSmoothingSpeed);
+
+		float BudgetPressureInterpAlpha = FMath::Clamp((SmoothedBudgetPressure - GBudgetFactorBeforeAggressiveReducedWork) * 0.5f, 0.0f, 1.0f);
+		int32 StateChangeThrottleInFrames = (int32)FMath::Lerp(4.0f, (float)GStateChangeThrottleInFrames, BudgetPressureInterpAlpha);
+
+		SET_FLOAT_STAT(STAT_AnimationBudgetAllocator_SmoothedBudgetPressure, SmoothedBudgetPressure);
+
 		// Queue for tick
 		for (SortedComponentIndex = 0; SortedComponentIndex < TotalIdealWorkUnits; ++SortedComponentIndex)
 		{
@@ -779,7 +787,7 @@ int32 FAnimationBudgetAllocator::CalculateWorkDistributionAndQueue(float InDelta
 				}
 			}
 
-			QueueForTick(ComponentData);
+			QueueForTick(ComponentData, StateChangeThrottleInFrames);
 		}
 
 		// If any components are not longer allowed to perform reduced work, force them back out
@@ -796,15 +804,8 @@ int32 FAnimationBudgetAllocator::CalculateWorkDistributionAndQueue(float InDelta
 			}
 		}
 
-		const float BudgetPressure = (float)TotalIdealWorkUnits / WorkUnitBudget;
-		SmoothedBudgetPressure = FMath::FInterpTo(SmoothedBudgetPressure, BudgetPressure, InDeltaSeconds, GBudgetPressureSmoothingSpeed);
-
-		SET_FLOAT_STAT(STAT_AnimationBudgetAllocator_SmoothedBudgetPressure, SmoothedBudgetPressure);
-
 		if(--ReducedComponentWorkCounter <= 0)
 		{
-			float BudgetPressureInterpAlpha = FMath::Clamp((SmoothedBudgetPressure - GBudgetFactorBeforeAggressiveReducedWork) * 0.5f, 0.0f, 1.0f);
-
 			// Scale num components to switch based on budget pressure
 			const int32 NumComponentsToSwitch = (int32)FMath::Lerp(1.0f, (float)GReducedWorkThrottleMaxPerFrame, BudgetPressureInterpAlpha);
 			int32 ComponentsSwitched = 0;
diff --git a/Engine/Plugins/Runtime/AnimationBudgetAllocator/Source/AnimationBudgetAllocator/Private/AnimationBudgetAllocator.h b/Engine/Plugins/Runtime/AnimationBudgetAllocator/Source/AnimationBudgetAllocator/Private/AnimationBudgetAllocator.h
--- a/Engine/Plugins/Runtime/AnimationBudgetAllocator/Source/AnimationBudgetAllocator/Private/AnimationBudgetAllocator.h
+++ b/Engine/Plugins/Runtime/AnimationBudgetAllocator/Source/AnimationBudgetAllocator/Private/AnimationBudgetAllocator.h
@@ -94,7 +94,7 @@ public:
 	uint8 SkippedTicks;
 
 	/** Counter used to prevent state changes from happening too often */
-	uint8 StateChangeThrottle;
+	int8 StateChangeThrottle;
 
 	/** Whether we ever tick */
 	uint8 bTickEnabled : 1;
