Added CSV exclusive stat for animation budget allocator

Renamed plain ol' stat as well. as no longer hand 'initial' and 'secondary' ticks.

diff --git a/Engine/Plugins/Runtime/AnimationBudgetAllocator/Source/AnimationBudgetAllocator/Private/AnimationBudgetAllocator.cpp b/Engine/Plugins/Runtime/AnimationBudgetAllocator/Source/AnimationBudgetAllocator/Private/AnimationBudgetAllocator.cpp
--- a/Engine/Plugins/Runtime/AnimationBudgetAllocator/Source/AnimationBudgetAllocator/Private/AnimationBudgetAllocator.cpp
+++ b/Engine/Plugins/Runtime/AnimationBudgetAllocator/Source/AnimationBudgetAllocator/Private/AnimationBudgetAllocator.cpp
@@ -10,7 +10,7 @@
 #include "SkeletalMeshComponentBudgeted.h"
 #include "DrawDebugHelpers.h"
 
-DECLARE_CYCLE_STAT(TEXT("InitialTick"), STAT_AnimationBudgetAllocator_InitialTick, STATGROUP_AnimationBudgetAllocator);
+DECLARE_CYCLE_STAT(TEXT("InitialTick"), STAT_AnimationBudgetAllocator_Update, STATGROUP_AnimationBudgetAllocator);
 
 DECLARE_DWORD_COUNTER_STAT(TEXT("Num Registered Components"), STAT_AnimationBudgetAllocator_NumRegisteredComponents, STATGROUP_AnimationBudgetAllocator);
 DECLARE_DWORD_COUNTER_STAT(TEXT("Num Ticked Components"), STAT_AnimationBudgetAllocator_NumTickedComponents, STATGROUP_AnimationBudgetAllocator);
@@ -18,7 +18,6 @@ DECLARE_DWORD_COUNTER_STAT(TEXT("Num Ticked Components"), STAT_AnimationBudgetAl
 DECLARE_DWORD_COUNTER_STAT(TEXT("Demand"), STAT_AnimationBudgetAllocator_Demand, STATGROUP_AnimationBudgetAllocator);
 DECLARE_FLOAT_COUNTER_STAT(TEXT("Budget"), STAT_AnimationBudgetAllocator_Budget, STATGROUP_AnimationBudgetAllocator);
 DECLARE_FLOAT_COUNTER_STAT(TEXT("Average Work Unit (ms)"), STAT_AnimationBudgetAllocator_AverageWorkUnitTime, STATGROUP_AnimationBudgetAllocator);
-DECLARE_FLOAT_COUNTER_STAT(TEXT("Initial Tick (ms)"), STAT_AnimationBudgetAllocator_InitialTickMs, STATGROUP_AnimationBudgetAllocator);
 DECLARE_DWORD_COUNTER_STAT(TEXT("Always Tick"), STAT_AnimationBudgetAllocator_AlwaysTick, STATGROUP_AnimationBudgetAllocator);
 DECLARE_DWORD_COUNTER_STAT(TEXT("Throttled"), STAT_AnimationBudgetAllocator_Throttled, STATGROUP_AnimationBudgetAllocator);
 DECLARE_DWORD_COUNTER_STAT(TEXT("Interpolated"), STAT_AnimationBudgetAllocator_Interpolated, STATGROUP_AnimationBudgetAllocator);
@@ -886,7 +885,8 @@ void FAnimationBudgetAllocator::OnWorldPreActorTick(UWorld* InWorld, ELevelTick
 
 void FAnimationBudgetAllocator::Update(float DeltaSeconds)
 {
-	SCOPE_CYCLE_COUNTER(STAT_AnimationBudgetAllocator_InitialTick);
+	SCOPE_CYCLE_COUNTER(STAT_AnimationBudgetAllocator_Update);
+	CSV_SCOPED_TIMING_STAT_EXCLUSIVE(AnimationBudgetAllocator);
 
 	FAnimationBudgetAllocator::bCachedEnabled = GAnimationBudgetEnabled == 1 && bEnabled;
 
