Animation budget: Add emergency reduction budget pressure level

When budget pressure gets too high, we relax some of the restrictions around moving to reduced work. While this is in place, all non-always-tick components will be running reduced work.

diff --git a/Engine/Plugins/Runtime/AnimationBudgetAllocator/Source/AnimationBudgetAllocator/Private/AnimationBudgetAllocator.cpp b/Engine/Plugins/Runtime/AnimationBudgetAllocator/Source/AnimationBudgetAllocator/Private/AnimationBudgetAllocator.cpp
--- a/Engine/Plugins/Runtime/AnimationBudgetAllocator/Source/AnimationBudgetAllocator/Private/AnimationBudgetAllocator.cpp
+++ b/Engine/Plugins/Runtime/AnimationBudgetAllocator/Source/AnimationBudgetAllocator/Private/AnimationBudgetAllocator.cpp
@@ -289,6 +289,19 @@ static FAutoConsoleVariableRef CVarSkelBatch_ReducedWorkThrottleMaxPerFrame(
 	}),
 	ECVF_Scalability);
 
+static float GBudgetPressureBeforeEmergencyReducedWork = 2.5f;
+
+static FAutoConsoleVariableRef CVarSkelBatch_BudgetPressureBeforeEmergencyReducedWork(
+	TEXT("a.Budget.GBudgetPressureBeforeEmergencyReducedWork"),
+	GBudgetPressureBeforeEmergencyReducedWork,
+	TEXT("Range > 0.0, Default = 2.5\n")
+	TEXT("Controls the budget pressure where emergency reduced work (applied to all components except those that are bAlwaysTick).\n"),
+	FConsoleVariableDelegate::CreateLambda([](IConsoleVariable* InVariable)
+	{
+		GBudgetPressureBeforeEmergencyReducedWork = FMath::Max(GBudgetPressureBeforeEmergencyReducedWork, 0.0f);
+	}),
+	ECVF_Scalability);
+
 FComponentData::FComponentData(USkeletalMeshComponentBudgeted* InComponent)
 	: Component(InComponent)
 	, RootPrerequisite(nullptr)
@@ -805,6 +818,8 @@ int32 FAnimationBudgetAllocator::CalculateWorkDistributionAndQueue(float InDelta
 
 		if(--ReducedComponentWorkCounter <= 0)
 		{
+			const bool bEmergencyReducedWork = SmoothedBudgetPressure >= GBudgetPressureBeforeEmergencyReducedWork;
+
 			// Scale num components to switch based on budget pressure
 			const int32 NumComponentsToSwitch = (int32)FMath::Lerp(1.0f, (float)GReducedWorkThrottleMaxPerFrame, BudgetPressureInterpAlpha);
 			int32 ComponentsSwitched = 0;
@@ -837,7 +852,10 @@ int32 FAnimationBudgetAllocator::CalculateWorkDistributionAndQueue(float InDelta
 				for (SortedComponentIndex = TotalIdealWorkUnits - 1; SortedComponentIndex >= FullIndexEnd; --SortedComponentIndex)
 				{
 					FComponentData& ComponentData = AllComponentData[AllSortedComponentData[SortedComponentIndex]];
-					if(ComponentData.bAllowReducedWork && !ComponentData.bReducedWork && ComponentData.Component->OnReduceWork().IsBound())
+
+					const bool bAllowReducedWork = (ComponentData.bAllowReducedWork || bEmergencyReducedWork) && !ComponentData.bAlwaysTick;
+
+					if(bAllowReducedWork && !ComponentData.bReducedWork && ComponentData.Component->OnReduceWork().IsBound())
 					{
 #if WITH_TICK_DEBUG
 						UE_LOG(LogTemp, Warning, TEXT("Reducing component work (mesh %s) (actor %llx)"), ComponentData.Component->SkeletalMesh ? *ComponentData.Component->SkeletalMesh->GetName() : TEXT("null"), (uint64)ComponentData.Component->GetOwner());
