BudgetAllocator optionally looks for actor's recently rendered flag.

diff --git a/Engine/Plugins/Runtime/AnimationBudgetAllocator/Source/AnimationBudgetAllocator/Private/AnimationBudgetAllocator.cpp b/Engine/Plugins/Runtime/AnimationBudgetAllocator/Source/AnimationBudgetAllocator/Private/AnimationBudgetAllocator.cpp
--- a/Engine/Plugins/Runtime/AnimationBudgetAllocator/Source/AnimationBudgetAllocator/Private/AnimationBudgetAllocator.cpp
+++ b/Engine/Plugins/Runtime/AnimationBudgetAllocator/Source/AnimationBudgetAllocator/Private/AnimationBudgetAllocator.cpp
@@ -499,6 +499,7 @@ void FAnimationBudgetAllocator::QueueSortedComponentIndices(float InDeltaSeconds
 				auto ShouldComponentTick = [WorldTime](const USkeletalMeshComponentBudgeted* InComponent, const FComponentData& InComponentData)
 				{
 					return ((InComponent->LastRenderTime > WorldTime) ||
+						(InComponent->GetShouldUseActorRenderedFlag() && InComponent->GetAttachmentRootActor() && InComponent->GetAttachmentRootActor()->WasRecentlyRendered())  ||
 							InComponentData.bTickEvenIfNotRendered ||
 							InComponent->ShouldTickPose() ||
 							InComponent->ShouldUpdateTransform(false) ||	// We can force this to false, only used in WITH_EDITOR
diff --git a/Engine/Plugins/Runtime/AnimationBudgetAllocator/Source/AnimationBudgetAllocator/Private/SkeletalMeshComponentBudgeted.cpp b/Engine/Plugins/Runtime/AnimationBudgetAllocator/Source/AnimationBudgetAllocator/Private/SkeletalMeshComponentBudgeted.cpp
--- a/Engine/Plugins/Runtime/AnimationBudgetAllocator/Source/AnimationBudgetAllocator/Private/SkeletalMeshComponentBudgeted.cpp
+++ b/Engine/Plugins/Runtime/AnimationBudgetAllocator/Source/AnimationBudgetAllocator/Private/SkeletalMeshComponentBudgeted.cpp
@@ -15,6 +15,7 @@ USkeletalMeshComponentBudgeted::USkeletalMeshComponentBudgeted(const FObjectInit
 	, AnimationBudgetAllocator(nullptr)
 	, bAutoRegisterWithBudgetAllocator(true)
 	, bAutoCalculateSignificance(false)
+	, bShouldUseActorRenderedFlag(false)
 {
 }
 
diff --git a/Engine/Plugins/Runtime/AnimationBudgetAllocator/Source/AnimationBudgetAllocator/Public/SkeletalMeshComponentBudgeted.h b/Engine/Plugins/Runtime/AnimationBudgetAllocator/Source/AnimationBudgetAllocator/Public/SkeletalMeshComponentBudgeted.h
--- a/Engine/Plugins/Runtime/AnimationBudgetAllocator/Source/AnimationBudgetAllocator/Public/SkeletalMeshComponentBudgeted.h
+++ b/Engine/Plugins/Runtime/AnimationBudgetAllocator/Source/AnimationBudgetAllocator/Public/SkeletalMeshComponentBudgeted.h
@@ -41,6 +41,9 @@ public:
 	/** Get delegate called to calculate significance if bAutoCalculateSignificance = true */
 	static FOnCalculateSignificance& OnCalculateSignificance() { return OnCalculateSignificanceDelegate; }
 
+	bool GetShouldUseActorRenderedFlag() const { return bShouldUseActorRenderedFlag; };
+
+	void SetShouldUseActorRenderedFlag(bool value) { bShouldUseActorRenderedFlag = value; };
 private:
 	// UActorComponent interface
 	virtual void BeginPlay() override;
@@ -79,4 +82,7 @@ private:
 	/** Whether this component should automatically register with the budget allocator in OnRegister/OnUnregister */
 	UPROPERTY(EditAnywhere, Category = Budgeting)
 	uint8 bAutoCalculateSignificance : 1;
+
+	UPROPERTY(EditAnywhere, Category = Budgeting)
+	uint8 bShouldUseActorRenderedFlag : 1;
 };
\ No newline at end of file
