Misc bugfixes & improvements to animation budget allocator

Added force params to allow debugging of frame skip/interpolation issues.
Fixed rigid body node being continually reset when running external frame skips.

diff --git a/Engine/Plugins/Runtime/AnimationBudgetAllocator/Source/AnimationBudgetAllocator/Private/AnimationBudgetAllocator.cpp b/Engine/Plugins/Runtime/AnimationBudgetAllocator/Source/AnimationBudgetAllocator/Private/AnimationBudgetAllocator.cpp
--- a/Engine/Plugins/Runtime/AnimationBudgetAllocator/Source/AnimationBudgetAllocator/Private/AnimationBudgetAllocator.cpp
+++ b/Engine/Plugins/Runtime/AnimationBudgetAllocator/Source/AnimationBudgetAllocator/Private/AnimationBudgetAllocator.cpp
@@ -368,6 +368,7 @@ int32 FAnimationBudgetAllocator::CalculateWorkDistributionAndQueue(float InDelta
 			InComponentData.Component->EnableExternalUpdate(bTickThisFrame);
 			InComponentData.Component->EnableExternalEvaluationRateLimiting(InComponentData.TickRate > 1);
 			InComponentData.Component->SetExternalDeltaTime(InComponentData.AccumulatedDeltaTime);
+			InComponentData.Component->SetExternalTickRate(InComponentData.TickRate);
 
 			InComponentData.AccumulatedDeltaTime = bTickThisFrame ? 0.0f : InComponentData.AccumulatedDeltaTime;
 
@@ -394,224 +395,254 @@ int32 FAnimationBudgetAllocator::CalculateWorkDistributionAndQueue(float InDelta
 		}
 	};
 
-	const int32 TotalIdealWorkUnits = AllSortedComponentData.Num();
-
-	SET_DWORD_STAT(STAT_AnimationBudgetAllocator_Demand, TotalIdealWorkUnits);
-
-	if(TotalIdealWorkUnits > 0)
+#if !(UE_BUILD_SHIPPING || UE_BUILD_TEST)
+	if(GAnimationBudgetDebugForce)
 	{
-		// Calc smoothed average of last frames' work units
-		const float AverageTickTimeMs = TotalEstimatedTickTimeMs / NumWorkUnitsForAverage;
-		AverageWorkUnitTimeMs = FMath::FInterpTo(AverageWorkUnitTimeMs, AverageTickTimeMs, InDeltaSeconds, Parameters.WorkUnitSmoothingSpeed);
-
-		SET_FLOAT_STAT(STAT_AnimationBudgetAllocator_AverageWorkUnitTime, AverageWorkUnitTimeMs);
-		CSV_CUSTOM_STAT(AnimationBudget, AverageWorkUnitTimeMs, AverageTickTimeMs, ECsvCustomStatOp::Set);
-
-		// Want to map the remaining (non-fixed) work units so that we only execute N work units per frame.
-		// If we can go over budget to keep quality then we use that value
-		const float WorkUnitBudget = FMath::Max(Parameters.BudgetInMs / AverageWorkUnitTimeMs, (float)TotalIdealWorkUnits * Parameters.MinQuality);
-
-		SET_FLOAT_STAT(STAT_AnimationBudgetAllocator_Budget, WorkUnitBudget);
-
-		// Ramp-off work units that we tick every frame once required ticks start exceeding budget
-		const float WorkUnitsExcess = FMath::Max(0.0f, TotalIdealWorkUnits - WorkUnitBudget);
-		const float WorkUnitsToRunInFull = FMath::Clamp(WorkUnitBudget - (WorkUnitsExcess * Parameters.AlwaysTickFalloffAggression), (float)NumComponentsToNotSkip, (float)TotalIdealWorkUnits);
-		SET_DWORD_STAT(STAT_AnimationBudgetAllocator_AlwaysTick, WorkUnitsToRunInFull);
-		BUDGET_CSV_STAT(AnimationBudget, NumAlwaysTicked, WorkUnitsToRunInFull, ECsvCustomStatOp::Set);
-		const int32 FullIndexEnd = (int32)WorkUnitsToRunInFull;
-
-		// Account for the actual time that we think the fixed ticks will take
-		// This works better when budget to work unit ratio is low
-		float FullTickTime = 0.0f;
-		int32 SortedComponentIndex;
-		for (SortedComponentIndex = 0; SortedComponentIndex < FullIndexEnd; ++SortedComponentIndex)
+		for(int32 SortedComponentIndex : AllSortedComponentData)
 		{
-			FComponentData& ComponentData = AllComponentData[AllSortedComponentData[SortedComponentIndex]];
-			FullTickTime += ComponentData.GameThreadLastCompletionTimeMs + ComponentData.GameThreadLastTickTimeMs;
+			FComponentData& ComponentData = AllComponentData[SortedComponentIndex];
+			ComponentData.TickRate = ComponentData.DesiredTickRate = GAnimationBudgetDebugForceRate;
+			ComponentData.bInterpolate = !!GAnimationBudgetDebugForceInterpolation;
+			if(ComponentData.Component->OnReduceWork().IsBound())
+			{
+				if(ComponentData.bReducedWork && !GAnimationBudgetDebugForceReducedWork)
+				{
+					ComponentData.Component->OnReduceWork().Execute(ComponentData.Component, false);
+					ComponentData.bReducedWork = false;
+				}
+				else if(!ComponentData.bReducedWork && GAnimationBudgetDebugForceReducedWork)
+				{
+					ComponentData.Component->OnReduceWork().Execute(ComponentData.Component, true);
+					ComponentData.bReducedWork = true;
+				}
+			}
+			ComponentData.StateChangeThrottle = -1;
+
+			QueueForTick(ComponentData, -1);
 		}
+	}
+	else
+#endif
+	{
+		const int32 TotalIdealWorkUnits = AllSortedComponentData.Num();
 
-		float FullTickWorkUnits = FMath::Min(FullTickTime / AverageWorkUnitTimeMs, WorkUnitsToRunInFull);
+		SET_DWORD_STAT(STAT_AnimationBudgetAllocator_Demand, TotalIdealWorkUnits);
 
-		float RemainingBudget = FMath::Max(0.0f, WorkUnitBudget - FullTickWorkUnits);
-		float RemainingWorkUnitsToRun = FMath::Max(0.0f, TotalIdealWorkUnits - FullTickWorkUnits);
+		if(TotalIdealWorkUnits > 0)
+		{
+			// Calc smoothed average of last frames' work units
+			const float AverageTickTimeMs = TotalEstimatedTickTimeMs / NumWorkUnitsForAverage;
+			AverageWorkUnitTimeMs = FMath::FInterpTo(AverageWorkUnitTimeMs, AverageTickTimeMs, InDeltaSeconds, Parameters.WorkUnitSmoothingSpeed);
 
-		// Ramp off interpolated units in a similar way
-		const float WorkUnitsToInterpolate = FMath::Min(FMath::Max(RemainingBudget - (WorkUnitsExcess * Parameters.InterpolationFalloffAggression), (float)FMath::Min(Parameters.MaxInterpolatedComponents, NumComponentsToNotThrottle)), RemainingWorkUnitsToRun);
-		SET_DWORD_STAT(STAT_AnimationBudgetAllocator_Interpolated, WorkUnitsToInterpolate);
+			SET_FLOAT_STAT(STAT_AnimationBudgetAllocator_AverageWorkUnitTime, AverageWorkUnitTimeMs);
+			CSV_CUSTOM_STAT(AnimationBudget, AverageWorkUnitTimeMs, AverageTickTimeMs, ECsvCustomStatOp::Set);
 
-		const int32 InterpolationIndexEnd = FMath::Min((int32)WorkUnitsToInterpolate + (int32)WorkUnitsToRunInFull, TotalIdealWorkUnits);
+			// Want to map the remaining (non-fixed) work units so that we only execute N work units per frame.
+			// If we can go over budget to keep quality then we use that value
+			const float WorkUnitBudget = FMath::Max(Parameters.BudgetInMs / AverageWorkUnitTimeMs, (float)TotalIdealWorkUnits * Parameters.MinQuality);
 
-		const float MaxInterpolationRate = (float)Parameters.InterpolationMaxRate;
+			SET_FLOAT_STAT(STAT_AnimationBudgetAllocator_Budget, WorkUnitBudget);
 
-		// Calc remaining (throttled) work units
-		RemainingBudget = FMath::Max(0.0f, RemainingBudget - (WorkUnitsToInterpolate * Parameters.InterpolationTickMultiplier));
-		RemainingWorkUnitsToRun = FMath::Max(0.0f, RemainingWorkUnitsToRun - (WorkUnitsToInterpolate * Parameters.InterpolationTickMultiplier));
+			// Ramp-off work units that we tick every frame once required ticks start exceeding budget
+			const float WorkUnitsExcess = FMath::Max(0.0f, TotalIdealWorkUnits - WorkUnitBudget);
+			const float WorkUnitsToRunInFull = FMath::Clamp(WorkUnitBudget - (WorkUnitsExcess * Parameters.AlwaysTickFalloffAggression), (float)NumComponentsToNotSkip, (float)TotalIdealWorkUnits);
+			SET_DWORD_STAT(STAT_AnimationBudgetAllocator_AlwaysTick, WorkUnitsToRunInFull);
+			BUDGET_CSV_STAT(AnimationBudget, NumAlwaysTicked, WorkUnitsToRunInFull, ECsvCustomStatOp::Set);
+			const int32 FullIndexEnd = (int32)WorkUnitsToRunInFull;
 
-		SET_DWORD_STAT(STAT_AnimationBudgetAllocator_Throttled, RemainingWorkUnitsToRun);
+			// Account for the actual time that we think the fixed ticks will take
+			// This works better when budget to work unit ratio is low
+			float FullTickTime = 0.0f;
+			int32 SortedComponentIndex;
+			for (SortedComponentIndex = 0; SortedComponentIndex < FullIndexEnd; ++SortedComponentIndex)
+			{
+				FComponentData& ComponentData = AllComponentData[AllSortedComponentData[SortedComponentIndex]];
+				FullTickTime += ComponentData.GameThreadLastCompletionTimeMs + ComponentData.GameThreadLastTickTimeMs;
+			}
 
-		// Midpoint of throttle gradient is RemainingWorkUnitsToRun / RemainingBudget.
-		// If we distributed this as a constant we would get each component ticked
-		// at the same rate. However we want to tick more significant meshes more often,
-		// so we keep the area under the curve constant and intercept the line with this centroid.
-		// Care must be taken with rounding to keep workload in-budget.
-		const float ThrottleRateDenominator = RemainingBudget > 1.0f ? RemainingBudget : 1.0f;
-		const float MaxThrottleRate = FMath::Min(FMath::CeilToFloat(FMath::Max(1.0f, RemainingWorkUnitsToRun / ThrottleRateDenominator) * 2.0f), (float)Parameters.MaxTickRate);
-		const float ThrottleDenominator = RemainingWorkUnitsToRun > 0.0f ? RemainingWorkUnitsToRun : 1.0f;
+			float FullTickWorkUnits = FMath::Min(FullTickTime / AverageWorkUnitTimeMs, WorkUnitsToRunInFull);
 
-		// Bucket 1: always ticked
-		for (SortedComponentIndex = 0; SortedComponentIndex < FullIndexEnd; ++SortedComponentIndex)
-		{
-			FComponentData& ComponentData = AllComponentData[AllSortedComponentData[SortedComponentIndex]];
+			float RemainingBudget = FMath::Max(0.0f, WorkUnitBudget - FullTickWorkUnits);
+			float RemainingWorkUnitsToRun = FMath::Max(0.0f, TotalIdealWorkUnits - FullTickWorkUnits);
 
-			// not skipping frames here as we can either match demand or these components need a full update
-			ComponentData.TickRate = 1;
-			ComponentData.DesiredTickRate = 1;
-			ComponentData.bInterpolate = false;
-		}
+			// Ramp off interpolated units in a similar way
+			const float WorkUnitsToInterpolate = FMath::Min(FMath::Max(RemainingBudget - (WorkUnitsExcess * Parameters.InterpolationFalloffAggression), (float)FMath::Min(Parameters.MaxInterpolatedComponents, NumComponentsToNotThrottle)), RemainingWorkUnitsToRun);
+			SET_DWORD_STAT(STAT_AnimationBudgetAllocator_Interpolated, WorkUnitsToInterpolate);
 
-		// Bucket 2: interpolated
-		int32 NumInterpolated = 0;
-		for (SortedComponentIndex = FullIndexEnd; SortedComponentIndex < InterpolationIndexEnd; ++SortedComponentIndex)
-		{
-			FComponentData& ComponentData = AllComponentData[AllSortedComponentData[SortedComponentIndex]];
+			const int32 InterpolationIndexEnd = FMath::Min((int32)WorkUnitsToInterpolate + (int32)WorkUnitsToRunInFull, TotalIdealWorkUnits);
 
-			const float Alpha = (((float)SortedComponentIndex - FullIndexEnd) / WorkUnitsToInterpolate);
-			ComponentData.DesiredTickRate = FMath::Min((int32)FMath::FloorToFloat(FMath::Lerp(2.0f, MaxInterpolationRate, Alpha) + 0.5f), 255);
-			ComponentData.bInterpolate = true;
-			NumInterpolated++;
-		}
+			const float MaxInterpolationRate = (float)Parameters.InterpolationMaxRate;
 
-		// Bucket 3: Rate limited
-		int32 NumThrottled = 0;
-		for (SortedComponentIndex = InterpolationIndexEnd; SortedComponentIndex < TotalIdealWorkUnits; ++SortedComponentIndex)
-		{
-			FComponentData& ComponentData = AllComponentData[AllSortedComponentData[SortedComponentIndex]];
+			// Calc remaining (throttled) work units
+			RemainingBudget = FMath::Max(0.0f, RemainingBudget - (WorkUnitsToInterpolate * Parameters.InterpolationTickMultiplier));
+			RemainingWorkUnitsToRun = FMath::Max(0.0f, RemainingWorkUnitsToRun - (WorkUnitsToInterpolate * Parameters.InterpolationTickMultiplier));
 
-			const float Alpha = (((float)SortedComponentIndex - InterpolationIndexEnd) / ThrottleDenominator);
-			ComponentData.DesiredTickRate = FMath::Min((int32)FMath::FloorToFloat(FMath::Lerp(2.0f, MaxThrottleRate, Alpha) + 0.5f), 255);
-			ComponentData.bInterpolate = false;
-			NumThrottled++;
-		}
+			SET_DWORD_STAT(STAT_AnimationBudgetAllocator_Throttled, RemainingWorkUnitsToRun);
 
-		BUDGET_CSV_STAT(AnimationBudget, NumInterpolated, NumInterpolated, ECsvCustomStatOp::Set);
-		BUDGET_CSV_STAT(AnimationBudget, NumThrottled, RemainingWorkUnitsToRun, ECsvCustomStatOp::Set);
+			// Midpoint of throttle gradient is RemainingWorkUnitsToRun / RemainingBudget.
+			// If we distributed this as a constant we would get each component ticked
+			// at the same rate. However we want to tick more significant meshes more often,
+			// so we keep the area under the curve constant and intercept the line with this centroid.
+			// Care must be taken with rounding to keep workload in-budget.
+			const float ThrottleRateDenominator = RemainingBudget > 1.0f ? RemainingBudget : 1.0f;
+			const float MaxThrottleRate = FMath::Min(FMath::CeilToFloat(FMath::Max(1.0f, RemainingWorkUnitsToRun / ThrottleRateDenominator) * 2.0f), (float)Parameters.MaxTickRate);
+			const float ThrottleDenominator = RemainingWorkUnitsToRun > 0.0f ? RemainingWorkUnitsToRun : 1.0f;
 
-		const float BudgetPressure = (float)TotalIdealWorkUnits / WorkUnitBudget;
-		SmoothedBudgetPressure = FMath::FInterpTo(SmoothedBudgetPressure, BudgetPressure, InDeltaSeconds, Parameters.BudgetPressureSmoothingSpeed);
+			// Bucket 1: always ticked
+			for (SortedComponentIndex = 0; SortedComponentIndex < FullIndexEnd; ++SortedComponentIndex)
+			{
+				FComponentData& ComponentData = AllComponentData[AllSortedComponentData[SortedComponentIndex]];
 
-		float BudgetPressureInterpAlpha = FMath::Clamp((SmoothedBudgetPressure - Parameters.BudgetFactorBeforeAggressiveReducedWork) * 0.5f, 0.0f, 1.0f);
-		int32 StateChangeThrottleInFrames = (int32)FMath::Lerp(4.0f, (float)Parameters.StateChangeThrottleInFrames, BudgetPressureInterpAlpha);
+				// not skipping frames here as we can either match demand or these components need a full update
+				ComponentData.TickRate = 1;
+				ComponentData.DesiredTickRate = 1;
+				ComponentData.bInterpolate = false;
+			}
 
-		SET_FLOAT_STAT(STAT_AnimationBudgetAllocator_SmoothedBudgetPressure, SmoothedBudgetPressure);
+			// Bucket 2: interpolated
+			int32 NumInterpolated = 0;
+			for (SortedComponentIndex = FullIndexEnd; SortedComponentIndex < InterpolationIndexEnd; ++SortedComponentIndex)
+			{
+				FComponentData& ComponentData = AllComponentData[AllSortedComponentData[SortedComponentIndex]];
 
-		// Queue for tick
-		for (SortedComponentIndex = 0; SortedComponentIndex < TotalIdealWorkUnits; ++SortedComponentIndex)
-		{
-			FComponentData& ComponentData = AllComponentData[AllSortedComponentData[SortedComponentIndex]];
+				const float Alpha = (((float)SortedComponentIndex - FullIndexEnd) / WorkUnitsToInterpolate);
+				ComponentData.DesiredTickRate = FMath::Min((int32)FMath::FloorToFloat(FMath::Lerp(2.0f, MaxInterpolationRate, Alpha) + 0.5f), 255);
+				ComponentData.bInterpolate = true;
+				NumInterpolated++;
+			}
 
-			// Ensure that root prerequisite doesnt end up with a lower (or different) tick rate than dependencies
-			if(ComponentData.RootPrerequisite != nullptr)
+			// Bucket 3: Rate limited
+			int32 NumThrottled = 0;
+			for (SortedComponentIndex = InterpolationIndexEnd; SortedComponentIndex < TotalIdealWorkUnits; ++SortedComponentIndex)
 			{
-				const int32 PrerequisiteHandle = ComponentData.RootPrerequisite->GetAnimationBudgetHandle();
-				if(PrerequisiteHandle != INDEX_NONE)
+				FComponentData& ComponentData = AllComponentData[AllSortedComponentData[SortedComponentIndex]];
+
+				const float Alpha = (((float)SortedComponentIndex - InterpolationIndexEnd) / ThrottleDenominator);
+				ComponentData.DesiredTickRate = FMath::Min((int32)FMath::FloorToFloat(FMath::Lerp(2.0f, MaxThrottleRate, Alpha) + 0.5f), 255);
+				ComponentData.bInterpolate = false;
+				NumThrottled++;
+			}
+
+			BUDGET_CSV_STAT(AnimationBudget, NumInterpolated, NumInterpolated, ECsvCustomStatOp::Set);
+			BUDGET_CSV_STAT(AnimationBudget, NumThrottled, RemainingWorkUnitsToRun, ECsvCustomStatOp::Set);
+
+			const float BudgetPressure = (float)TotalIdealWorkUnits / WorkUnitBudget;
+			SmoothedBudgetPressure = FMath::FInterpTo(SmoothedBudgetPressure, BudgetPressure, InDeltaSeconds, Parameters.BudgetPressureSmoothingSpeed);
+
+			float BudgetPressureInterpAlpha = FMath::Clamp((SmoothedBudgetPressure - Parameters.BudgetFactorBeforeAggressiveReducedWork) * 0.5f, 0.0f, 1.0f);
+			int32 StateChangeThrottleInFrames = (int32)FMath::Lerp(4.0f, (float)Parameters.StateChangeThrottleInFrames, BudgetPressureInterpAlpha);
+
+			SET_FLOAT_STAT(STAT_AnimationBudgetAllocator_SmoothedBudgetPressure, SmoothedBudgetPressure);
+
+			// Queue for tick
+			for (SortedComponentIndex = 0; SortedComponentIndex < TotalIdealWorkUnits; ++SortedComponentIndex)
+			{
+				FComponentData& ComponentData = AllComponentData[AllSortedComponentData[SortedComponentIndex]];
+
+				// Ensure that root prerequisite doesnt end up with a lower (or different) tick rate than dependencies
+				if(ComponentData.RootPrerequisite != nullptr)
 				{
-					FComponentData& RootPrerequisiteComponentData = AllComponentData[PrerequisiteHandle];
-					RootPrerequisiteComponentData.TickRate = ComponentData.TickRate = FMath::Min(ComponentData.TickRate, RootPrerequisiteComponentData.TickRate);
-					RootPrerequisiteComponentData.DesiredTickRate = ComponentData.DesiredTickRate = FMath::Min(ComponentData.DesiredTickRate, RootPrerequisiteComponentData.DesiredTickRate);
-					RootPrerequisiteComponentData.StateChangeThrottle = ComponentData.StateChangeThrottle = FMath::Min(ComponentData.StateChangeThrottle, RootPrerequisiteComponentData.StateChangeThrottle);
+					const int32 PrerequisiteHandle = ComponentData.RootPrerequisite->GetAnimationBudgetHandle();
+					if(PrerequisiteHandle != INDEX_NONE)
+					{
+						FComponentData& RootPrerequisiteComponentData = AllComponentData[PrerequisiteHandle];
+						RootPrerequisiteComponentData.TickRate = ComponentData.TickRate = FMath::Min(ComponentData.TickRate, RootPrerequisiteComponentData.TickRate);
+						RootPrerequisiteComponentData.DesiredTickRate = ComponentData.DesiredTickRate = FMath::Min(ComponentData.DesiredTickRate, RootPrerequisiteComponentData.DesiredTickRate);
+						RootPrerequisiteComponentData.StateChangeThrottle = ComponentData.StateChangeThrottle = FMath::Min(ComponentData.StateChangeThrottle, RootPrerequisiteComponentData.StateChangeThrottle);
+					}
 				}
-			}
 
-			QueueForTick(ComponentData, StateChangeThrottleInFrames);
-		}
+				QueueForTick(ComponentData, StateChangeThrottleInFrames);
+			}
 
-		// If any components are not longer allowed to perform reduced work, force them back out
-		for(int32 DisallowedReducedWorkComponentIndex : DisallowedReducedWorkComponentData)
-		{
-			FComponentData& ComponentData = AllComponentData[DisallowedReducedWorkComponentIndex];
-			if(ComponentData.bReducedWork && ComponentData.Component->OnReduceWork().IsBound())
+			// If any components are not longer allowed to perform reduced work, force them back out
+			for(int32 DisallowedReducedWorkComponentIndex : DisallowedReducedWorkComponentData)
 			{
+				FComponentData& ComponentData = AllComponentData[DisallowedReducedWorkComponentIndex];
+				if(ComponentData.bReducedWork && ComponentData.Component->OnReduceWork().IsBound())
+				{
 #if WITH_TICK_DEBUG
-				UE_LOG(LogTemp, Warning, TEXT("Force-increasing component work (mesh %s) (actor %llx)"), ComponentData.Component->SkeletalMesh ? *ComponentData.Component->SkeletalMesh->GetName() : TEXT("null"), (uint64)ComponentData.Component->GetOwner());
+					UE_LOG(LogTemp, Warning, TEXT("Force-increasing component work (mesh %s) (actor %llx)"), ComponentData.Component->SkeletalMesh ? *ComponentData.Component->SkeletalMesh->GetName() : TEXT("null"), (uint64)ComponentData.Component->GetOwner());
 #endif
-				ComponentData.Component->OnReduceWork().Execute(ComponentData.Component, false);
-				ComponentData.bReducedWork = false;
+					ComponentData.Component->OnReduceWork().Execute(ComponentData.Component, false);
+					ComponentData.bReducedWork = false;
+				}
 			}
-		}
 
-		if(--ReducedComponentWorkCounter <= 0)
-		{
-			const bool bEmergencyReducedWork = SmoothedBudgetPressure >= Parameters.BudgetPressureBeforeEmergencyReducedWork;
+			if(--ReducedComponentWorkCounter <= 0)
+			{
+				const bool bEmergencyReducedWork = SmoothedBudgetPressure >= Parameters.BudgetPressureBeforeEmergencyReducedWork;
 
-			// Scale num components to switch based on budget pressure
-			const int32 NumComponentsToSwitch = (int32)FMath::Lerp(1.0f, (float)Parameters.ReducedWorkThrottleMaxPerFrame, BudgetPressureInterpAlpha);
-			int32 ComponentsSwitched = 0;
+				// Scale num components to switch based on budget pressure
+				const int32 NumComponentsToSwitch = (int32)FMath::Lerp(1.0f, (float)Parameters.ReducedWorkThrottleMaxPerFrame, BudgetPressureInterpAlpha);
+				int32 ComponentsSwitched = 0;
 
-			// If we have any components running reduced work when we have an excess, then move them out of the 'reduced' pool per tick
-			if (ReducedWorkComponentData.Num() > 0 && SmoothedBudgetPressure < Parameters.BudgetFactorBeforeReducedWork - Parameters.BudgetFactorBeforeReducedWorkEpsilon)
-			{
-				for(int32 ReducedWorkComponentIndex : ReducedWorkComponentData)
+				// If we have any components running reduced work when we have an excess, then move them out of the 'reduced' pool per tick
+				if (ReducedWorkComponentData.Num() > 0 && SmoothedBudgetPressure < Parameters.BudgetFactorBeforeReducedWork - Parameters.BudgetFactorBeforeReducedWorkEpsilon)
 				{
-					FComponentData& ComponentData = AllComponentData[ReducedWorkComponentIndex];
-					if(ComponentData.bReducedWork && ComponentData.Component->OnReduceWork().IsBound())
+					for(int32 ReducedWorkComponentIndex : ReducedWorkComponentData)
 					{
+						FComponentData& ComponentData = AllComponentData[ReducedWorkComponentIndex];
+						if(ComponentData.bReducedWork && ComponentData.Component->OnReduceWork().IsBound())
+						{
 #if WITH_TICK_DEBUG
-						UE_LOG(LogTemp, Warning, TEXT("Increasing component work (mesh %s) (actor %llx)"), ComponentData.Component->SkeletalMesh ? *ComponentData.Component->SkeletalMesh->GetName() : TEXT("null"), (uint64)ComponentData.Component->GetOwner());
+							UE_LOG(LogTemp, Warning, TEXT("Increasing component work (mesh %s) (actor %llx)"), ComponentData.Component->SkeletalMesh ? *ComponentData.Component->SkeletalMesh->GetName() : TEXT("null"), (uint64)ComponentData.Component->GetOwner());
 #endif
-						ComponentData.Component->OnReduceWork().Execute(ComponentData.Component, false);
-						ComponentData.bReducedWork = false;
-
-						ComponentsSwitched++;
-						if(ComponentsSwitched >= NumComponentsToSwitch)
-						{
-							break;	
+							ComponentData.Component->OnReduceWork().Execute(ComponentData.Component, false);
+							ComponentData.bReducedWork = false;
+
+							ComponentsSwitched++;
+							if(ComponentsSwitched >= NumComponentsToSwitch)
+							{
+								break;	
+							}
 						}
 					}
 				}
-			}
-			else if(SmoothedBudgetPressure > Parameters.BudgetFactorBeforeReducedWork)
-			{
-				// Any work units that we interpolate or throttle should also be eligible for work reduction (which can involve disabling other ticks), so set them all now if needed
-				for (SortedComponentIndex = TotalIdealWorkUnits - 1; SortedComponentIndex >= FullIndexEnd; --SortedComponentIndex)
+				else if(SmoothedBudgetPressure > Parameters.BudgetFactorBeforeReducedWork)
 				{
-					FComponentData& ComponentData = AllComponentData[AllSortedComponentData[SortedComponentIndex]];
+					// Any work units that we interpolate or throttle should also be eligible for work reduction (which can involve disabling other ticks), so set them all now if needed
+					for (SortedComponentIndex = TotalIdealWorkUnits - 1; SortedComponentIndex >= FullIndexEnd; --SortedComponentIndex)
+					{
+						FComponentData& ComponentData = AllComponentData[AllSortedComponentData[SortedComponentIndex]];
 
-					const bool bAllowReducedWork = (ComponentData.bAllowReducedWork || bEmergencyReducedWork) && !ComponentData.bAlwaysTick;
+						const bool bAllowReducedWork = (ComponentData.bAllowReducedWork || bEmergencyReducedWork) && !ComponentData.bAlwaysTick;
 
-					if(bAllowReducedWork && !ComponentData.bReducedWork && ComponentData.Component->OnReduceWork().IsBound())
-					{
+						if(bAllowReducedWork && !ComponentData.bReducedWork && ComponentData.Component->OnReduceWork().IsBound())
+						{
 #if WITH_TICK_DEBUG
-						UE_LOG(LogTemp, Warning, TEXT("Reducing component work (mesh %s) (actor %llx)"), ComponentData.Component->SkeletalMesh ? *ComponentData.Component->SkeletalMesh->GetName() : TEXT("null"), (uint64)ComponentData.Component->GetOwner());
+							UE_LOG(LogTemp, Warning, TEXT("Reducing component work (mesh %s) (actor %llx)"), ComponentData.Component->SkeletalMesh ? *ComponentData.Component->SkeletalMesh->GetName() : TEXT("null"), (uint64)ComponentData.Component->GetOwner());
 #endif
-						ComponentData.Component->OnReduceWork().Execute(ComponentData.Component, true);
-						ComponentData.bReducedWork = true;
-
-						ComponentsSwitched++;
-						if(ComponentsSwitched >= NumComponentsToSwitch)
-						{
-							break;	
+							ComponentData.Component->OnReduceWork().Execute(ComponentData.Component, true);
+							ComponentData.bReducedWork = true;
+
+							ComponentsSwitched++;
+							if(ComponentsSwitched >= NumComponentsToSwitch)
+							{
+								break;	
+							}
 						}
 					}
 				}
-			}
 
-			// Scale the rate at which we consider reducing component work based on budget pressure
-			ReducedComponentWorkCounter = (int32)FMath::Lerp((float)Parameters.ReducedWorkThrottleMaxInFrames, (float)Parameters.ReducedWorkThrottleMinInFrames, BudgetPressureInterpAlpha);
+				// Scale the rate at which we consider reducing component work based on budget pressure
+				ReducedComponentWorkCounter = (int32)FMath::Lerp((float)Parameters.ReducedWorkThrottleMaxInFrames, (float)Parameters.ReducedWorkThrottleMinInFrames, BudgetPressureInterpAlpha);
+			}
 		}
-	}
 
 #if CSV_PROFILER
-	if(AllSortedComponentData.Num() > 0)
-	{
-		for (int32 ComponentDataIndex : AllSortedComponentData)
+		if(AllSortedComponentData.Num() > 0)
 		{
-			FComponentData& ComponentData = AllComponentData[ComponentDataIndex];
-			OutAverageTickRate += (float)ComponentData.TickRate;
-		}
+			for (int32 ComponentDataIndex : AllSortedComponentData)
+			{
+				FComponentData& ComponentData = AllComponentData[ComponentDataIndex];
+				OutAverageTickRate += (float)ComponentData.TickRate;
+			}
 
-		OutAverageTickRate /= (float)AllSortedComponentData.Num();
-	}
+			OutAverageTickRate /= (float)AllSortedComponentData.Num();
+		}
 #endif
+	}
 
 	return NumTicked;
 }
diff --git a/Engine/Plugins/Runtime/AnimationBudgetAllocator/Source/AnimationBudgetAllocator/Private/AnimationBudgetAllocatorCVars.cpp b/Engine/Plugins/Runtime/AnimationBudgetAllocator/Source/AnimationBudgetAllocator/Private/AnimationBudgetAllocatorCVars.cpp
--- a/Engine/Plugins/Runtime/AnimationBudgetAllocator/Source/AnimationBudgetAllocator/Private/AnimationBudgetAllocatorCVars.cpp
+++ b/Engine/Plugins/Runtime/AnimationBudgetAllocator/Source/AnimationBudgetAllocator/Private/AnimationBudgetAllocatorCVars.cpp
@@ -33,6 +33,54 @@ static FAutoConsoleVariableRef CVarSkelBatch_ShowAddresses(
 	ECVF_Scalability);
 #endif
 
+#if !(UE_BUILD_SHIPPING || UE_BUILD_TEST)
+
+int32 GAnimationBudgetDebugForce = 0;
+
+static FAutoConsoleVariableRef CVarSkelBatch_Force(
+	TEXT("a.Budget.Debug.Force"),
+	GAnimationBudgetDebugForce,
+	TEXT("Values: 0/1\n")
+	TEXT("Default: 0\n")
+	TEXT("Turns on forced rate/interp/reduced controls. These override any budget-driven values."),
+	ECVF_Scalability);
+
+int32 GAnimationBudgetDebugForceRate = 4;
+
+static FAutoConsoleVariableRef CVarSkelBatch_ForceRate(
+	TEXT("a.Budget.Debug.Force.Rate"),
+	GAnimationBudgetDebugForceRate,
+	TEXT("Values: >= 1\n")
+	TEXT("Default: 4\n")
+	TEXT("Forces all components to update at the specifed rate when a.Budget.Debug.Force is enabled."),
+	FConsoleVariableDelegate::CreateLambda([](IConsoleVariable* InVariable)
+	{
+		GAnimationBudgetDebugForceRate = FMath::Max(GAnimationBudgetDebugForceRate, 1);
+	}),
+	ECVF_Scalability);
+
+int32 GAnimationBudgetDebugForceInterpolation = 0;
+
+static FAutoConsoleVariableRef CVarSkelBatch_ForceInterpolation(
+	TEXT("a.Budget.Debug.Force.Interp"),
+	GAnimationBudgetDebugForceInterpolation,
+	TEXT("Values: 0/1\n")
+	TEXT("Default: 0\n")
+	TEXT("Forces interpolation on when a.Budget.Debug.Force is enabled."),
+	ECVF_Scalability);
+
+int32 GAnimationBudgetDebugForceReducedWork = 0;
+
+static FAutoConsoleVariableRef CVarSkelBatch_ForceReducedWork(
+	TEXT("a.Budget.Debug.Force.Reduced"),
+	GAnimationBudgetDebugForceReducedWork,
+	TEXT("Values: 0/1\n")
+	TEXT("Default: 0\n")
+	TEXT("Forces reduced work on when a.Budget.Debug.Force is enabled."),
+	ECVF_Scalability);
+
+#endif
+
 /** CVar-driven parameter block */
 FAnimationBudgetAllocatorParameters GBudgetParameters;
 
diff --git a/Engine/Plugins/Runtime/AnimationBudgetAllocator/Source/AnimationBudgetAllocator/Private/AnimationBudgetAllocatorCVars.h b/Engine/Plugins/Runtime/AnimationBudgetAllocator/Source/AnimationBudgetAllocator/Private/AnimationBudgetAllocatorCVars.h
--- a/Engine/Plugins/Runtime/AnimationBudgetAllocator/Source/AnimationBudgetAllocator/Private/AnimationBudgetAllocatorCVars.h
+++ b/Engine/Plugins/Runtime/AnimationBudgetAllocator/Source/AnimationBudgetAllocator/Private/AnimationBudgetAllocatorCVars.h
@@ -18,6 +18,16 @@ extern int32 GAnimationBudgetDebugEnabled;
 extern int32 GAnimationBudgetDebugShowAddresses;
 #endif
 
+#if !(UE_BUILD_SHIPPING || UE_BUILD_TEST)
+
+/** Parameters used to force state for debugging */
+extern int32 GAnimationBudgetDebugForce;
+extern int32 GAnimationBudgetDebugForceRate;
+extern int32 GAnimationBudgetDebugForceInterpolation;
+extern int32 GAnimationBudgetDebugForceReducedWork;
+
+#endif
+
 /** CVar-driven parameter block */
 extern FAnimationBudgetAllocatorParameters GBudgetParameters;
 
diff --git a/Engine/Source/Runtime/AnimGraphRuntime/Private/BoneControllers/AnimNode_RigidBody.cpp b/Engine/Source/Runtime/AnimGraphRuntime/Private/BoneControllers/AnimNode_RigidBody.cpp
--- a/Engine/Source/Runtime/AnimGraphRuntime/Private/BoneControllers/AnimNode_RigidBody.cpp
+++ b/Engine/Source/Runtime/AnimGraphRuntime/Private/BoneControllers/AnimNode_RigidBody.cpp
@@ -220,9 +220,14 @@ void FAnimNode_RigidBody::EvaluateSkeletalControl_AnyThread(FComponentSpacePoseC
 	//SCOPED_NAMED_EVENT_TEXT("FAnimNode_Ragdoll::EvaluateSkeletalControl_AnyThread", FColor::Magenta);
 
 	// Update our eval counter, and decide whether we need to reset simulated bodies, if our anim instance hasn't updated in a while.
-	if(EvalCounter.HasEverBeenUpdated() && !EvalCounter.WasSynchronizedLastFrame(Output.AnimInstanceProxy->GetEvaluationCounter()))
+	if(EvalCounter.HasEverBeenUpdated())
 	{
-		ResetSimulatedTeleportType = ETeleportType::ResetPhysics;
+		// Always propagate skip rate as it can go up and down between updates
+		EvalCounter.SetMaxSkippedFrames(Output.AnimInstanceProxy->GetEvaluationCounter().GetMaxSkippedFrames());
+		if(!EvalCounter.WasSynchronizedLastFrame(Output.AnimInstanceProxy->GetEvaluationCounter()))
+		{
+			ResetSimulatedTeleportType = ETeleportType::ResetPhysics;
+		}
 	}
 	EvalCounter.SynchronizeWith(Output.AnimInstanceProxy->GetEvaluationCounter());
 
diff --git a/Engine/Source/Runtime/Engine/Classes/Components/SkinnedMeshComponent.h b/Engine/Source/Runtime/Engine/Classes/Components/SkinnedMeshComponent.h
--- a/Engine/Source/Runtime/Engine/Classes/Components/SkinnedMeshComponent.h
+++ b/Engine/Source/Runtime/Engine/Classes/Components/SkinnedMeshComponent.h
@@ -430,6 +430,10 @@ public:
 	uint8& MeshComponentUpdateFlag = *(uint8*)&VisibilityBasedAnimTickOption;
 #endif
 
+protected:
+	/** Record of the tick rate we are using when externally controlled */
+	uint8 ExternalTickRate;
+
 protected:
 	/** used to cache previous bone transform or not */
 	uint8 bHasValidBoneTransform:1;
@@ -572,6 +576,15 @@ public:
 	/** Set whether we have our tick rate externally controlled non-URO-based interpolation */
 	void EnableExternalTickRateControl(bool bInEnable) { bExternalTickRateControlled = bInEnable; }
 
+	/** Check whether we we have our tick rate externally controlled */
+	bool IsUsingExternalTickRateControl() const { return bExternalTickRateControlled; }
+
+	/** Set the external tick rate */
+	void SetExternalTickRate(uint8 InTickRate) { ExternalTickRate = InTickRate; }
+
+	/** Get the external tick rate */
+	uint8 GetExternalTickRate() const { return ExternalTickRate; }
+
 	/** Enable non-URO-based interpolation */
 	void EnableExternalInterpolation(bool bInEnable) { bExternalInterpolate = bInEnable; }
 
diff --git a/Engine/Source/Runtime/Engine/Private/Animation/AnimInstanceProxy.cpp b/Engine/Source/Runtime/Engine/Private/Animation/AnimInstanceProxy.cpp
--- a/Engine/Source/Runtime/Engine/Private/Animation/AnimInstanceProxy.cpp
+++ b/Engine/Source/Runtime/Engine/Private/Animation/AnimInstanceProxy.cpp
@@ -486,17 +486,20 @@ void FAnimInstanceProxy::InitializeObjects(UAnimInstance* InAnimInstance)
 	NumUroSkippedFrames_Eval = 0;
 	if(RateParams)
 	{
-		bool bDoUro = SkeletalMeshComponent->ShouldUseUpdateRateOptimizations();
-		bool bDoEvalOptimization = bDoUro && RateParams->DoEvaluationRateOptimizations();
-
+		const bool bDoUro = SkeletalMeshComponent->ShouldUseUpdateRateOptimizations();
 		if(bDoUro)
 		{
 			NumUroSkippedFrames_Update = RateParams->UpdateRate - 1;
-		}
 
-		if(bDoEvalOptimization)
+			const bool bDoEvalOptimization = RateParams->DoEvaluationRateOptimizations();
+			if(bDoEvalOptimization)
+			{
+				NumUroSkippedFrames_Eval = RateParams->EvaluationRate - 1;
+			}
+		}
+		else if(SkeletalMeshComponent->IsUsingExternalTickRateControl())
 		{
-			NumUroSkippedFrames_Eval = RateParams->EvaluationRate - 1;
+			NumUroSkippedFrames_Update = NumUroSkippedFrames_Eval = SkeletalMeshComponent->GetExternalTickRate();
 		}
 	}
 
diff --git a/Engine/Source/Runtime/Engine/Private/Components/SkinnedMeshComponent.cpp b/Engine/Source/Runtime/Engine/Private/Components/SkinnedMeshComponent.cpp
--- a/Engine/Source/Runtime/Engine/Private/Components/SkinnedMeshComponent.cpp
+++ b/Engine/Source/Runtime/Engine/Private/Components/SkinnedMeshComponent.cpp
@@ -359,6 +359,7 @@ USkinnedMeshComponent::USkinnedMeshComponent(const FObjectInitializer& ObjectIni
 
 	ExternalInterpolationAlpha = 0.0f;
 	ExternalDeltaTime = 0.0f;
+	ExternalTickRate = 1;
 	bExternalInterpolate = false;
 	bExternalUpdate = false;
 	bExternalEvaluationRateLimited = false;
diff --git a/Engine/Source/Runtime/Engine/Public/Animation/AnimTypes.h b/Engine/Source/Runtime/Engine/Public/Animation/AnimTypes.h
--- a/Engine/Source/Runtime/Engine/Public/Animation/AnimTypes.h
+++ b/Engine/Source/Runtime/Engine/Public/Animation/AnimTypes.h
@@ -177,6 +177,12 @@ public:
 		SkipFrames = InMaxSkippedFrames;
 	}
 
+	/** Gets the number of expected frames to be skipped by URO */
+	int16 GetMaxSkippedFrames() const
+	{
+		return SkipFrames;
+	}
+
 	/** Clear the internal counters and frame skip */
 	void Reset()
 	{
