Relax SetComponentSignificance warning to only during play.  Fixes a warning during shutdown.

#rnx

diff --git a/Engine/Plugins/Runtime/AnimationBudgetAllocator/Source/AnimationBudgetAllocator/Private/SkeletalMeshComponentBudgeted.cpp b/Engine/Plugins/Runtime/AnimationBudgetAllocator/Source/AnimationBudgetAllocator/Private/SkeletalMeshComponentBudgeted.cpp
--- a/Engine/Plugins/Runtime/AnimationBudgetAllocator/Source/AnimationBudgetAllocator/Private/SkeletalMeshComponentBudgeted.cpp
+++ b/Engine/Plugins/Runtime/AnimationBudgetAllocator/Source/AnimationBudgetAllocator/Private/SkeletalMeshComponentBudgeted.cpp
@@ -59,7 +59,7 @@ void USkeletalMeshComponentBudgeted::SetComponentSignificance(float Significance
 	{
 		AnimationBudgetAllocator->SetComponentSignificance(this, Significance, bNeverSkip, bTickEvenIfNotRendered, bAllowReducedWork, bForceInterpolate);
 	}
-	else
+	else if (HasBegunPlay())
 	{
 		UE_LOG(LogSkeletalMesh, Warning, TEXT("SetComponentSignificance called on [%s] before registering with budget allocator"), *GetName());
 	}
