Fixed an Editor crash when clicking on the "Enables the simulation of the blueprint and ticking" button in a blueprint.

#rnx

diff --git a/Engine/Plugins/Runtime/AnimationBudgetAllocator/Source/AnimationBudgetAllocator/Private/SkeletalMeshComponentBudgeted.cpp b/Engine/Plugins/Runtime/AnimationBudgetAllocator/Source/AnimationBudgetAllocator/Private/SkeletalMeshComponentBudgeted.cpp
--- a/Engine/Plugins/Runtime/AnimationBudgetAllocator/Source/AnimationBudgetAllocator/Private/SkeletalMeshComponentBudgeted.cpp
+++ b/Engine/Plugins/Runtime/AnimationBudgetAllocator/Source/AnimationBudgetAllocator/Private/SkeletalMeshComponentBudgeted.cpp
@@ -25,7 +25,13 @@ void USkeletalMeshComponentBudgeted::BeginPlay()
 
 	if(bAutoRegisterWithBudgetAllocator && !UKismetSystemLibrary::IsDedicatedServer(this))
 	{
-		IAnimationBudgetAllocator::Get(GetWorld())->RegisterComponent(this);
+		if (UWorld* LocalWorld = GetWorld())
+		{
+			if (IAnimationBudgetAllocator* AnimationBudgetAllocator = IAnimationBudgetAllocator::Get(LocalWorld))
+			{
+				AnimationBudgetAllocator->RegisterComponent(this);
+			}
+		}
 	}
 }
 
@@ -35,7 +41,13 @@ void USkeletalMeshComponentBudgeted::EndPlay(const EEndPlayReason::Type EndPlayR
 	// As reciprocal ptrs are null, handles are all invalid.
 	if(!IsUnreachable())	
 	{
-		IAnimationBudgetAllocator::Get(GetWorld())->UnregisterComponent(this);
+		if (UWorld* LocalWorld = GetWorld())
+		{
+			if (IAnimationBudgetAllocator* AnimationBudgetAllocator = IAnimationBudgetAllocator::Get(LocalWorld))
+			{
+				AnimationBudgetAllocator->UnregisterComponent(this);
+			}
+		}
 	}
 
 	Super::EndPlay(EndPlayReason);
