Fixed an Editor crash when clicking on the "Enables the simulation of the blueprint and ticking" button in a blueprint.

#rnx

diff --git a/Engine/Plugins/Runtime/AnimationBudgetAllocator/Source/AnimationBudgetAllocator/Private/SkeletalMeshComponentBudgeted.cpp b/Engine/Plugins/Runtime/AnimationBudgetAllocator/Source/AnimationBudgetAllocator/Private/SkeletalMeshComponentBudgeted.cpp
--- a/Engine/Plugins/Runtime/AnimationBudgetAllocator/Source/AnimationBudgetAllocator/Private/SkeletalMeshComponentBudgeted.cpp
+++ b/Engine/Plugins/Runtime/AnimationBudgetAllocator/Source/AnimationBudgetAllocator/Private/SkeletalMeshComponentBudgeted.cpp
@@ -27,9 +27,9 @@ void USkeletalMeshComponentBudgeted::BeginPlay()
 	{
 		if (UWorld* LocalWorld = GetWorld())
 		{
-			if (IAnimationBudgetAllocator* AnimationBudgetAllocator = IAnimationBudgetAllocator::Get(LocalWorld))
+			if (IAnimationBudgetAllocator* LocalAnimationBudgetAllocator = IAnimationBudgetAllocator::Get(LocalWorld))
 			{
-				AnimationBudgetAllocator->RegisterComponent(this);
+				LocalAnimationBudgetAllocator->RegisterComponent(this);
 			}
 		}
 	}
@@ -43,9 +43,9 @@ void USkeletalMeshComponentBudgeted::EndPlay(const EEndPlayReason::Type EndPlayR
 	{
 		if (UWorld* LocalWorld = GetWorld())
 		{
-			if (IAnimationBudgetAllocator* AnimationBudgetAllocator = IAnimationBudgetAllocator::Get(LocalWorld))
+			if (IAnimationBudgetAllocator* LocalAnimationBudgetAllocator = IAnimationBudgetAllocator::Get(LocalWorld))
 			{
-				AnimationBudgetAllocator->UnregisterComponent(this);
+				LocalAnimationBudgetAllocator->UnregisterComponent(this);
 			}
 		}
 	}
