Renamed FComponentData in AnimationBudgetAllocator.h to FAnimBudgetAllocatorComponentData. Fix a crash that's suspected to be caused by the compiler generating code with the size of another struct with that name (there are multiple).
CL#11716928 fixed a similar issue and contains more details.

diff --git a/Engine/Plugins/Runtime/AnimationBudgetAllocator/Source/AnimationBudgetAllocator/Private/AnimationBudgetAllocator.cpp b/Engine/Plugins/Runtime/AnimationBudgetAllocator/Source/AnimationBudgetAllocator/Private/AnimationBudgetAllocator.cpp
--- a/Engine/Plugins/Runtime/AnimationBudgetAllocator/Source/AnimationBudgetAllocator/Private/AnimationBudgetAllocator.cpp
+++ b/Engine/Plugins/Runtime/AnimationBudgetAllocator/Source/AnimationBudgetAllocator/Private/AnimationBudgetAllocator.cpp
@@ -33,7 +33,7 @@ CSV_DEFINE_CATEGORY(AnimationBudget, true);
 
 bool FAnimationBudgetAllocator::bCachedEnabled = false;
 
-FComponentData::FComponentData(USkeletalMeshComponentBudgeted* InComponent, float InGameThreadLastTickTimeMs, int32 InStateChangeThrottle)
+FAnimBudgetAllocatorComponentData::FAnimBudgetAllocatorComponentData(USkeletalMeshComponentBudgeted* InComponent, float InGameThreadLastTickTimeMs, int32 InStateChangeThrottle)
 	: Component(InComponent)
 	, RootPrerequisite(nullptr)
 	, Significance(1.0f)
@@ -140,7 +140,7 @@ void FAnimationBudgetAllocator::SetComponentSignificance(USkeletalMeshComponentB
 		int32 Handle = Component->GetAnimationBudgetHandle();
 		if(Handle != INDEX_NONE)
 		{
-			FComponentData& ComponentData = AllComponentData[Handle];
+			FAnimBudgetAllocatorComponentData& ComponentData = AllComponentData[Handle];
 			ComponentData.Significance = Significance;
 			ComponentData.bAlwaysTick = bAlwaysTick;
 			ComponentData.bTickEvenIfNotRendered = bTickEvenIfNotRendered;
@@ -164,7 +164,7 @@ void FAnimationBudgetAllocator::QueueSortedComponentIndices(float InDeltaSeconds
 	TotalEstimatedTickTimeMs = 0.0f;
 	NumWorkUnitsForAverage = 0.0f;
 
-	auto QueueComponentTick = [InDeltaSeconds, this](FComponentData& InComponentData, int32 InComponentIndex, bool bInOnScreen)
+	auto QueueComponentTick = [InDeltaSeconds, this](FAnimBudgetAllocatorComponentData& InComponentData, int32 InComponentIndex, bool bInOnScreen)
 	{
 		InComponentData.AccumulatedDeltaTime += InDeltaSeconds;
 		InComponentData.bOnScreen = bInOnScreen;
@@ -198,7 +198,7 @@ void FAnimationBudgetAllocator::QueueSortedComponentIndices(float InDeltaSeconds
 		}
 	};
 
-	auto DisableComponentTick = [this](FComponentData& InComponentData)
+	auto DisableComponentTick = [this](FAnimBudgetAllocatorComponentData& InComponentData)
 	{
 		InComponentData.SkippedTicks = 0;
 		InComponentData.AccumulatedDeltaTime = 0.0f;
@@ -211,14 +211,14 @@ void FAnimationBudgetAllocator::QueueSortedComponentIndices(float InDeltaSeconds
 
 	uint8 MaxComponentTickFunctionIndex = 0;
 	int32 ComponentIndex = 0;
-	for (FComponentData& ComponentData : AllComponentData)
+	for (FAnimBudgetAllocatorComponentData& ComponentData : AllComponentData)
 	{
 		if(ComponentData.bTickEnabled)
 		{
 			USkeletalMeshComponentBudgeted* Component = ComponentData.Component;
 			if (Component && Component->IsRegistered())
 			{
-				auto ShouldComponentTick = [WorldTime](const USkeletalMeshComponentBudgeted* InComponent, const FComponentData& InComponentData)
+				auto ShouldComponentTick = [WorldTime](const USkeletalMeshComponentBudgeted* InComponent, const FAnimBudgetAllocatorComponentData& InComponentData)
 				{
 					return ((InComponent->GetLastRenderTime() > WorldTime) ||
 						(InComponent->GetShouldUseActorRenderedFlag() && InComponent->GetAttachmentRootActor() && InComponent->GetAttachmentRootActor()->WasRecentlyRendered())  ||
@@ -237,7 +237,7 @@ void FAnimationBudgetAllocator::QueueSortedComponentIndices(float InDeltaSeconds
 					const int32 PrerequisiteHandle = ComponentData.RootPrerequisite->GetAnimationBudgetHandle();
 					if(PrerequisiteHandle != INDEX_NONE)
 					{
-						const FComponentData& RootPrerequisiteComponentData = AllComponentData[PrerequisiteHandle];
+						const FAnimBudgetAllocatorComponentData& RootPrerequisiteComponentData = AllComponentData[PrerequisiteHandle];
 						bShouldTick &= ShouldComponentTick(ComponentData.RootPrerequisite, RootPrerequisiteComponentData);
 					}
 				}
@@ -280,8 +280,8 @@ void FAnimationBudgetAllocator::QueueSortedComponentIndices(float InDeltaSeconds
 	// Sort by significance, largest first
 	auto SignificanceSortPredicate = [this](int32 InIndex0, int32 InIndex1)
 	{
-		const FComponentData& ComponentData0 = AllComponentData[InIndex0];
-		const FComponentData& ComponentData1 = AllComponentData[InIndex1];
+		const FAnimBudgetAllocatorComponentData& ComponentData0 = AllComponentData[InIndex0];
+		const FAnimBudgetAllocatorComponentData& ComponentData1 = AllComponentData[InIndex1];
 		return ComponentData0.Significance > ComponentData1.Significance;
 	};
 
@@ -292,7 +292,7 @@ void FAnimationBudgetAllocator::QueueSortedComponentIndices(float InDeltaSeconds
 	const int32 MaxOffscreenComponents = FMath::Min(NonRenderedComponentData.Num(), Parameters.MaxTickedOffsreenComponents);
 	if(MaxOffscreenComponents > 0)
 	{
-		auto ReduceWorkForOffscreenComponent = [](FComponentData& InComponentData)
+		auto ReduceWorkForOffscreenComponent = [](FAnimBudgetAllocatorComponentData& InComponentData)
 		{
 			if(InComponentData.bAllowReducedWork && !InComponentData.bReducedWork && InComponentData.Component->OnReduceWork().IsBound())
 			{
@@ -308,7 +308,7 @@ void FAnimationBudgetAllocator::QueueSortedComponentIndices(float InDeltaSeconds
 		int32 NonRenderedComponentIndex = 0;
 		for(; NonRenderedComponentIndex < MaxOffscreenComponents; ++NonRenderedComponentIndex)
 		{
-			FComponentData& ComponentData = AllComponentData[NonRenderedComponentData[NonRenderedComponentIndex]];
+			FAnimBudgetAllocatorComponentData& ComponentData = AllComponentData[NonRenderedComponentData[NonRenderedComponentIndex]];
 			QueueComponentTick(ComponentData, NonRenderedComponentData[NonRenderedComponentIndex], false);
 
 			// Always move to reduced work offscreen
@@ -321,7 +321,7 @@ void FAnimationBudgetAllocator::QueueSortedComponentIndices(float InDeltaSeconds
 		// Disable ticks for the rest
 		for(; NonRenderedComponentIndex < NonRenderedComponentData.Num(); ++NonRenderedComponentIndex)
 		{
-			FComponentData& ComponentData = AllComponentData[NonRenderedComponentData[NonRenderedComponentIndex]];
+			FAnimBudgetAllocatorComponentData& ComponentData = AllComponentData[NonRenderedComponentData[NonRenderedComponentIndex]];
 			DisableComponentTick(AllComponentData[NonRenderedComponentData[NonRenderedComponentIndex]]);
 
 			// Always move to reduced work offscreen
@@ -336,7 +336,7 @@ void FAnimationBudgetAllocator::QueueSortedComponentIndices(float InDeltaSeconds
 	}
 
 #if WITH_TICK_DEBUG
-	auto SignificanceDebugSortPredicate = [](const FComponentData& InComponentData0, const FComponentData& InComponentData1)
+	auto SignificanceDebugSortPredicate = [](const FAnimBudgetAllocatorComponentData& InComponentData0, const FAnimBudgetAllocatorComponentData& InComponentData1)
 	{
 		return InComponentData0.Significance > InComponentData1.Significance;
 	};
@@ -349,10 +349,10 @@ int32 FAnimationBudgetAllocator::CalculateWorkDistributionAndQueue(float InDelta
 {
 	int32 NumTicked = 0;
 
-	auto QueueForTick = [&NumTicked, this](FComponentData& InComponentData, int32 InStateChangeThrottleInFrames)
+	auto QueueForTick = [&NumTicked, this](FAnimBudgetAllocatorComponentData& InComponentData, int32 InStateChangeThrottleInFrames)
 	{
 		const int32 PrerequisiteHandle = InComponentData.RootPrerequisite != nullptr ? InComponentData.RootPrerequisite->GetAnimationBudgetHandle() : INDEX_NONE;
-		const FComponentData& ComponentDataToCheck = PrerequisiteHandle != INDEX_NONE ? AllComponentData[PrerequisiteHandle] : InComponentData;
+		const FAnimBudgetAllocatorComponentData& ComponentDataToCheck = PrerequisiteHandle != INDEX_NONE ? AllComponentData[PrerequisiteHandle] : InComponentData;
 
 		// Using (frame offset + frame counter) % tick rate allows us to only tick at the specified interval,
 		// but at a roughly even distribution over all registered components
@@ -401,7 +401,7 @@ int32 FAnimationBudgetAllocator::CalculateWorkDistributionAndQueue(float InDelta
 	{
 		for(int32 SortedComponentIndex : AllSortedComponentData)
 		{
-			FComponentData& ComponentData = AllComponentData[SortedComponentIndex];
+			FAnimBudgetAllocatorComponentData& ComponentData = AllComponentData[SortedComponentIndex];
 			ComponentData.TickRate = ComponentData.DesiredTickRate = GAnimationBudgetDebugForceRate;
 			ComponentData.bInterpolate = !!GAnimationBudgetDebugForceInterpolation;
 			if(ComponentData.Component->OnReduceWork().IsBound())
@@ -457,7 +457,7 @@ int32 FAnimationBudgetAllocator::CalculateWorkDistributionAndQueue(float InDelta
 			int32 SortedComponentIndex;
 			for (SortedComponentIndex = 0; SortedComponentIndex < FullIndexEnd; ++SortedComponentIndex)
 			{
-				FComponentData& ComponentData = AllComponentData[AllSortedComponentData[SortedComponentIndex]];
+				FAnimBudgetAllocatorComponentData& ComponentData = AllComponentData[AllSortedComponentData[SortedComponentIndex]];
 				FullTickTime += ComponentData.GameThreadLastCompletionTimeMs + ComponentData.GameThreadLastTickTimeMs;
 			}
 
@@ -492,7 +492,7 @@ int32 FAnimationBudgetAllocator::CalculateWorkDistributionAndQueue(float InDelta
 			// Bucket 1: always ticked
 			for (SortedComponentIndex = 0; SortedComponentIndex < FullIndexEnd; ++SortedComponentIndex)
 			{
-				FComponentData& ComponentData = AllComponentData[AllSortedComponentData[SortedComponentIndex]];
+				FAnimBudgetAllocatorComponentData& ComponentData = AllComponentData[AllSortedComponentData[SortedComponentIndex]];
 
 				// not skipping frames here as we can either match demand or these components need a full update
 				ComponentData.TickRate = 1;
@@ -504,7 +504,7 @@ int32 FAnimationBudgetAllocator::CalculateWorkDistributionAndQueue(float InDelta
 			int32 NumInterpolated = 0;
 			for (SortedComponentIndex = FullIndexEnd; SortedComponentIndex < InterpolationIndexEnd; ++SortedComponentIndex)
 			{
-				FComponentData& ComponentData = AllComponentData[AllSortedComponentData[SortedComponentIndex]];
+				FAnimBudgetAllocatorComponentData& ComponentData = AllComponentData[AllSortedComponentData[SortedComponentIndex]];
 
 				const float Alpha = (((float)SortedComponentIndex - FullIndexEnd) / WorkUnitsToInterpolate);
 				ComponentData.DesiredTickRate = FMath::Min((int32)FMath::FloorToFloat(FMath::Lerp(2.0f, MaxInterpolationRate, Alpha) + 0.5f), 255);
@@ -516,7 +516,7 @@ int32 FAnimationBudgetAllocator::CalculateWorkDistributionAndQueue(float InDelta
 			int32 NumThrottled = 0;
 			for (SortedComponentIndex = InterpolationIndexEnd; SortedComponentIndex < TotalIdealWorkUnits; ++SortedComponentIndex)
 			{
-				FComponentData& ComponentData = AllComponentData[AllSortedComponentData[SortedComponentIndex]];
+				FAnimBudgetAllocatorComponentData& ComponentData = AllComponentData[AllSortedComponentData[SortedComponentIndex]];
 
 				const float Alpha = (((float)SortedComponentIndex - InterpolationIndexEnd) / ThrottleDenominator);
 				ComponentData.DesiredTickRate = FMath::Min((int32)FMath::FloorToFloat(FMath::Lerp(2.0f, MaxThrottleRate, Alpha) + 0.5f), 255);
@@ -538,7 +538,7 @@ int32 FAnimationBudgetAllocator::CalculateWorkDistributionAndQueue(float InDelta
 			// Queue for tick
 			for (SortedComponentIndex = 0; SortedComponentIndex < TotalIdealWorkUnits; ++SortedComponentIndex)
 			{
-				FComponentData& ComponentData = AllComponentData[AllSortedComponentData[SortedComponentIndex]];
+				FAnimBudgetAllocatorComponentData& ComponentData = AllComponentData[AllSortedComponentData[SortedComponentIndex]];
 
 				// Ensure that root prerequisite doesnt end up with a lower (or different) tick rate than dependencies
 				if(ComponentData.RootPrerequisite != nullptr)
@@ -546,7 +546,7 @@ int32 FAnimationBudgetAllocator::CalculateWorkDistributionAndQueue(float InDelta
 					const int32 PrerequisiteHandle = ComponentData.RootPrerequisite->GetAnimationBudgetHandle();
 					if(PrerequisiteHandle != INDEX_NONE)
 					{
-						FComponentData& RootPrerequisiteComponentData = AllComponentData[PrerequisiteHandle];
+						FAnimBudgetAllocatorComponentData& RootPrerequisiteComponentData = AllComponentData[PrerequisiteHandle];
 						RootPrerequisiteComponentData.TickRate = ComponentData.TickRate = FMath::Min(ComponentData.TickRate, RootPrerequisiteComponentData.TickRate);
 						RootPrerequisiteComponentData.DesiredTickRate = ComponentData.DesiredTickRate = FMath::Min(ComponentData.DesiredTickRate, RootPrerequisiteComponentData.DesiredTickRate);
 						RootPrerequisiteComponentData.StateChangeThrottle = ComponentData.StateChangeThrottle = FMath::Min(ComponentData.StateChangeThrottle, RootPrerequisiteComponentData.StateChangeThrottle);
@@ -559,7 +559,7 @@ int32 FAnimationBudgetAllocator::CalculateWorkDistributionAndQueue(float InDelta
 			// If any components are not longer allowed to perform reduced work, force them back out
 			for(int32 DisallowedReducedWorkComponentIndex : DisallowedReducedWorkComponentData)
 			{
-				FComponentData& ComponentData = AllComponentData[DisallowedReducedWorkComponentIndex];
+				FAnimBudgetAllocatorComponentData& ComponentData = AllComponentData[DisallowedReducedWorkComponentIndex];
 				if(ComponentData.bReducedWork && ComponentData.Component->OnReduceWork().IsBound())
 				{
 #if WITH_TICK_DEBUG
@@ -583,7 +583,7 @@ int32 FAnimationBudgetAllocator::CalculateWorkDistributionAndQueue(float InDelta
 				{
 					for(int32 ReducedWorkComponentIndex : ReducedWorkComponentData)
 					{
-						FComponentData& ComponentData = AllComponentData[ReducedWorkComponentIndex];
+						FAnimBudgetAllocatorComponentData& ComponentData = AllComponentData[ReducedWorkComponentIndex];
 						if(ComponentData.bReducedWork && ComponentData.Component->OnReduceWork().IsBound())
 						{
 #if WITH_TICK_DEBUG
@@ -605,7 +605,7 @@ int32 FAnimationBudgetAllocator::CalculateWorkDistributionAndQueue(float InDelta
 					// Any work units that we interpolate or throttle should also be eligible for work reduction (which can involve disabling other ticks), so set them all now if needed
 					for (SortedComponentIndex = TotalIdealWorkUnits - 1; SortedComponentIndex >= FullIndexEnd; --SortedComponentIndex)
 					{
-						FComponentData& ComponentData = AllComponentData[AllSortedComponentData[SortedComponentIndex]];
+						FAnimBudgetAllocatorComponentData& ComponentData = AllComponentData[AllSortedComponentData[SortedComponentIndex]];
 
 						const bool bAllowReducedWork = (ComponentData.bAllowReducedWork || bEmergencyReducedWork) && !ComponentData.bAlwaysTick;
 
@@ -636,7 +636,7 @@ int32 FAnimationBudgetAllocator::CalculateWorkDistributionAndQueue(float InDelta
 		{
 			for (int32 ComponentDataIndex : AllSortedComponentData)
 			{
-				FComponentData& ComponentData = AllComponentData[ComponentDataIndex];
+				FAnimBudgetAllocatorComponentData& ComponentData = AllComponentData[ComponentDataIndex];
 				OutAverageTickRate += (float)ComponentData.TickRate;
 			}
 
@@ -694,7 +694,7 @@ void FAnimationBudgetAllocator::Update(float DeltaSeconds)
 			TMap<AActor*, TArray<int32>> ActorMap;
 			for (int32 ComponentDataIndex : AllSortedComponentData)
 			{
-				FComponentData& ComponentData = AllComponentData[ComponentDataIndex];
+				FAnimBudgetAllocatorComponentData& ComponentData = AllComponentData[ComponentDataIndex];
 				TArray<int32>& ComponentIndexArray = ActorMap.FindOrAdd(ComponentData.Component->GetOwner());
 				ComponentIndexArray.Add(ComponentDataIndex);
 			}
@@ -707,7 +707,7 @@ void FAnimationBudgetAllocator::Update(float DeltaSeconds)
 				
 				for(int32 ComponentDataIndex : ActorIndicesPair.Value)
 				{
-					FComponentData& ComponentData = AllComponentData[ComponentDataIndex];
+					FAnimBudgetAllocatorComponentData& ComponentData = AllComponentData[ComponentDataIndex];
 					if(ComponentData.bTickEnabled && ComponentData.bOnScreen)
 					{
 						if(GAnimationBudgetDebugShowAddresses != 0)
@@ -790,7 +790,7 @@ void FAnimationBudgetAllocator::OnHUDPostRender(AHUD* HUD, UCanvas* Canvas)
 				TMap<AActor*, TArray<int32>> ActorMap;
 				for (int32 ComponentDataIndex = 0; ComponentDataIndex < AllComponentData.Num(); ++ComponentDataIndex)
 				{
-					FComponentData& ComponentData = AllComponentData[ComponentDataIndex];
+					FAnimBudgetAllocatorComponentData& ComponentData = AllComponentData[ComponentDataIndex];
 					TArray<int32>& ComponentIndexArray = ActorMap.FindOrAdd(ComponentData.Component->GetOwner());
 					ComponentIndexArray.Add(ComponentDataIndex);
 				}
@@ -816,7 +816,7 @@ void FAnimationBudgetAllocator::OnHUDPostRender(AHUD* HUD, UCanvas* Canvas)
 
 					for(int32 ComponentDataIndex : ActorIndicesPair.Value)
 					{
-						FComponentData& ComponentData = AllComponentData[ComponentDataIndex];
+						FAnimBudgetAllocatorComponentData& ComponentData = AllComponentData[ComponentDataIndex];
 						if(ComponentData.bTickEnabled && !ComponentData.bOnScreen)
 						{
 							if(GAnimationBudgetDebugShowAddresses != 0)
@@ -906,7 +906,7 @@ void FAnimationBudgetAllocator::RegisterComponent(USkeletalMeshComponentBudgeted
 			InComponent->SetAnimationBudgetHandle(AllComponentData.Num());
 
 			// Setup frame offset
-			FComponentData& ComponentData = AllComponentData.Emplace_GetRef(InComponent, Parameters.InitialEstimatedWorkUnitTimeMs, Parameters.StateChangeThrottleInFrames);
+			FAnimBudgetAllocatorComponentData& ComponentData = AllComponentData.Emplace_GetRef(InComponent, Parameters.InitialEstimatedWorkUnitTimeMs, Parameters.StateChangeThrottleInFrames);
 			USkeletalMeshComponentBudgeted* RootPrerequisite = FindRootPrerequisite(InComponent);
 			ComponentData.RootPrerequisite = (RootPrerequisite != nullptr && RootPrerequisite != InComponent) ? RootPrerequisite : nullptr;
 			ComponentData.FrameOffset = CurrentFrameOffset++;
@@ -944,7 +944,7 @@ void FAnimationBudgetAllocator::UpdateComponentTickPrerequsites(USkeletalMeshCom
 		int32 ManagerHandle = InComponent->GetAnimationBudgetHandle();
 		if(ManagerHandle != INDEX_NONE)
 		{
-			FComponentData& ComponentData = AllComponentData[ManagerHandle];
+			FAnimBudgetAllocatorComponentData& ComponentData = AllComponentData[ManagerHandle];
 			USkeletalMeshComponentBudgeted* RootPrerequisite = FindRootPrerequisite(InComponent);
 			ComponentData.RootPrerequisite = (RootPrerequisite != nullptr && RootPrerequisite != InComponent) ? RootPrerequisite : nullptr;
 		}
@@ -955,7 +955,7 @@ void FAnimationBudgetAllocator::AddReferencedObjects(FReferenceCollector& Collec
 {
 	Collector.AddReferencedObject(World);
 
-	for (FComponentData& ComponentData : AllComponentData)
+	for (FAnimBudgetAllocatorComponentData& ComponentData : AllComponentData)
 	{
 		Collector.AddReferencedObject(ComponentData.Component);
 		Collector.AddReferencedObject(ComponentData.RootPrerequisite);
@@ -986,7 +986,7 @@ void FAnimationBudgetAllocator::SetGameThreadLastTickTimeMs(int32 InManagerHandl
 {
 	if(InManagerHandle != INDEX_NONE)
 	{
-		FComponentData& ComponentData = AllComponentData[InManagerHandle];
+		FAnimBudgetAllocatorComponentData& ComponentData = AllComponentData[InManagerHandle];
 		ComponentData.GameThreadLastTickTimeMs = InGameThreadLastTickTimeMs;
 #if ENABLE_DRAW_DEBUG
 		DebugTotalTime += InGameThreadLastTickTimeMs;
@@ -998,7 +998,7 @@ void FAnimationBudgetAllocator::SetGameThreadLastCompletionTimeMs(int32 InManage
 {
 	if(InManagerHandle != INDEX_NONE)
 	{
-		FComponentData& ComponentData = AllComponentData[InManagerHandle];
+		FAnimBudgetAllocatorComponentData& ComponentData = AllComponentData[InManagerHandle];
 		ComponentData.GameThreadLastCompletionTimeMs = InGameThreadLastCompletionTimeMs;
 #if ENABLE_DRAW_DEBUG
 		DebugTotalTime += InGameThreadLastCompletionTimeMs;
@@ -1013,7 +1013,7 @@ void FAnimationBudgetAllocator::SetIsRunningReducedWork(USkeletalMeshComponentBu
 		int32 ManagerHandle = InComponent->GetAnimationBudgetHandle();
 		if(ManagerHandle != INDEX_NONE)
 		{
-			FComponentData& ComponentData = AllComponentData[ManagerHandle];
+			FAnimBudgetAllocatorComponentData& ComponentData = AllComponentData[ManagerHandle];
 			ComponentData.bReducedWork = bInReducedWork;
 		}
 	}
@@ -1028,7 +1028,7 @@ void FAnimationBudgetAllocator::SetEnabled(bool bInEnabled)
 		// Remove all components we are currently tracking
 		for(int32 DataIndex = 0; DataIndex < AllComponentData.Num(); ++DataIndex)
 		{
-			FComponentData& ComponentData = AllComponentData[DataIndex];
+			FAnimBudgetAllocatorComponentData& ComponentData = AllComponentData[DataIndex];
 			if(ComponentData.Component != nullptr)
 			{
 				ComponentData.Component->SetAnimationBudgetHandle(INDEX_NONE);
diff --git a/Engine/Plugins/Runtime/AnimationBudgetAllocator/Source/AnimationBudgetAllocator/Private/AnimationBudgetAllocator.h b/Engine/Plugins/Runtime/AnimationBudgetAllocator/Source/AnimationBudgetAllocator/Private/AnimationBudgetAllocator.h
--- a/Engine/Plugins/Runtime/AnimationBudgetAllocator/Source/AnimationBudgetAllocator/Private/AnimationBudgetAllocator.h
+++ b/Engine/Plugins/Runtime/AnimationBudgetAllocator/Source/AnimationBudgetAllocator/Private/AnimationBudgetAllocator.h
@@ -36,9 +36,9 @@ class UCanvas;
 class FDebugDisplayInfo;
 
 /** Data for a single component */
-struct FComponentData
+struct FAnimBudgetAllocatorComponentData
 {
-	FComponentData()
+	FAnimBudgetAllocatorComponentData()
 		: Component(nullptr)
 		, RootPrerequisite(nullptr)
 		, Significance(0.0f)
@@ -61,9 +61,9 @@ struct FComponentData
 		, bNeverThrottle(true)
 	{}
 
-	FComponentData(USkeletalMeshComponentBudgeted* InComponent, float InGameThreadLastTickTimeMs, int32 InStateChangeThrottle);
+	FAnimBudgetAllocatorComponentData(USkeletalMeshComponentBudgeted* InComponent, float InGameThreadLastTickTimeMs, int32 InStateChangeThrottle);
 
-	bool operator==(const FComponentData& InOther) const
+	bool operator==(const FAnimBudgetAllocatorComponentData& InOther) const
 	{
 		return Component == InOther.Component;
 	}
@@ -192,7 +192,7 @@ protected:
 	UWorld* World;
 
 	// All component data
-	TArray<FComponentData> AllComponentData;
+	TArray<FAnimBudgetAllocatorComponentData> AllComponentData;
 
 	/** 
 	 * All currently tickable component indices sorted by significance, updated each tick.
@@ -202,7 +202,7 @@ protected:
 	TArray<int32> AllSortedComponentData;
 
 #if WITH_TICK_DEBUG
-	TArray<FComponentData*> AllSortedComponentDataDebug;
+	TArray<FAnimBudgetAllocatorComponentData*> AllSortedComponentDataDebug;
 #endif
 
 	/** All components that have reduced work that might want to tick (and hence might not want to do reduced work) */
