CIS fix for merge issue from CL#11998133 (FComponentData rename).
#rnx

diff --git a/Engine/Plugins/Runtime/AnimationBudgetAllocator/Source/AnimationBudgetAllocator/Private/AnimationBudgetAllocator.cpp b/Engine/Plugins/Runtime/AnimationBudgetAllocator/Source/AnimationBudgetAllocator/Private/AnimationBudgetAllocator.cpp
--- a/Engine/Plugins/Runtime/AnimationBudgetAllocator/Source/AnimationBudgetAllocator/Private/AnimationBudgetAllocator.cpp
+++ b/Engine/Plugins/Runtime/AnimationBudgetAllocator/Source/AnimationBudgetAllocator/Private/AnimationBudgetAllocator.cpp
@@ -802,8 +802,8 @@ void FAnimationBudgetAllocator::OnHUDPostRender(AHUD* HUD, UCanvas* Canvas)
 					// Sort by significance
 					ActorIndicesPair.Value.Sort([this](int32 Index0, int32 Index1)
 					{
-						const FComponentData& ComponentData0 = AllComponentData[Index0];
-						const FComponentData& ComponentData1 = AllComponentData[Index1];
+						const FAnimBudgetAllocatorComponentData& ComponentData0 = AllComponentData[Index0];
+						const FAnimBudgetAllocatorComponentData& ComponentData1 = AllComponentData[Index1];
  
 						return ComponentData0.Significance < ComponentData1.Significance;
 					});
