AnimationBudgetAlloc: Fixed interp alpha value calculation so it is 1 instead of 0.5 when tick rate is 1

diff --git a/Engine/Plugins/Runtime/AnimationBudgetAllocator/Source/AnimationBudgetAllocator/Private/AnimationBudgetAllocator.cpp b/Engine/Plugins/Runtime/AnimationBudgetAllocator/Source/AnimationBudgetAllocator/Private/AnimationBudgetAllocator.cpp
--- a/Engine/Plugins/Runtime/AnimationBudgetAllocator/Source/AnimationBudgetAllocator/Private/AnimationBudgetAllocator.cpp
+++ b/Engine/Plugins/Runtime/AnimationBudgetAllocator/Source/AnimationBudgetAllocator/Private/AnimationBudgetAllocator.cpp
@@ -356,7 +356,7 @@ int32 FAnimationBudgetAllocator::CalculateWorkDistributionAndQueue(float InDelta
 
 		// Calculate interp alpha even when interpolation or ticking is disabled as this is used to decide how much root motion to consume
 		// each frame by character movement
-		float Alpha = FMath::Clamp((1.0f / (InComponentData.TickRate - InComponentData.SkippedTicks + 1)), 0.0f, 1.0f);
+		float Alpha = (InComponentData.TickRate > 1) ? FMath::Clamp((1.0f / (InComponentData.TickRate - InComponentData.SkippedTicks + 1)), 0.0f, 1.0f) : 1.f;
 		InComponentData.Component->SetExternalInterpolationAlpha(Alpha);
 
 		// Using (frame offset + frame counter) % tick rate allows us to only tick at the specified interval,
