Make ExternalDeltaTime in AnimationBudgetAllocator use CustomTimeDialation

#tests Verified that various sized groups of actors using BudgetedAnimation properly animate as expected using default and modified CustomTimeDialation.

diff --git a/Engine/Plugins/Runtime/AnimationBudgetAllocator/Source/AnimationBudgetAllocator/Private/AnimationBudgetAllocator.cpp b/Engine/Plugins/Runtime/AnimationBudgetAllocator/Source/AnimationBudgetAllocator/Private/AnimationBudgetAllocator.cpp
--- a/Engine/Plugins/Runtime/AnimationBudgetAllocator/Source/AnimationBudgetAllocator/Private/AnimationBudgetAllocator.cpp
+++ b/Engine/Plugins/Runtime/AnimationBudgetAllocator/Source/AnimationBudgetAllocator/Private/AnimationBudgetAllocator.cpp
@@ -373,7 +373,9 @@ int32 FAnimationBudgetAllocator::CalculateWorkDistributionAndQueue(float InDelta
 			InComponentData.Component->EnableExternalInterpolation(InComponentData.TickRate > 1 && InComponentData.bInterpolate);
 			InComponentData.Component->EnableExternalUpdate(bTickThisFrame);
 			InComponentData.Component->EnableExternalEvaluationRateLimiting(InComponentData.TickRate > 1);
-			InComponentData.Component->SetExternalDeltaTime(InComponentData.AccumulatedDeltaTime);
+			AActor* MyOwner = InComponentData.Component->GetOwner();
+			const float CustomTimeDialation = MyOwner ? MyOwner->CustomTimeDilation : 1.0f;
+			InComponentData.Component->SetExternalDeltaTime(InComponentData.AccumulatedDeltaTime * CustomTimeDialation);
 			InComponentData.Component->SetExternalTickRate(InComponentData.TickRate);
 
 			InComponentData.AccumulatedDeltaTime = bTickThisFrame ? 0.0f : InComponentData.AccumulatedDeltaTime;
