[KillPendingKill] Fixed a case where budgeted SkeletalMeshComponents wouldn't unregister themselves from the animation budgeter, leaving a dangling strong reference. The unregister logic is being skipped because there were cases where the component reference was nulled due to being unreachable before it managed to remove itself. It would continue to try to remove itself using a stail index during teardown and either remove another component or past the index of the array, so we just detect this case now and skip removal.

Also making the post-GC routine a little more efficient by iterating from the back so we don't have to do it multiple times to account for skipped elements due to index shifts.

diff --git a/Engine/Plugins/Runtime/AnimationBudgetAllocator/Source/AnimationBudgetAllocator/Private/AnimationBudgetAllocator.cpp b/Engine/Plugins/Runtime/AnimationBudgetAllocator/Source/AnimationBudgetAllocator/Private/AnimationBudgetAllocator.cpp
--- a/Engine/Plugins/Runtime/AnimationBudgetAllocator/Source/AnimationBudgetAllocator/Private/AnimationBudgetAllocator.cpp
+++ b/Engine/Plugins/Runtime/AnimationBudgetAllocator/Source/AnimationBudgetAllocator/Private/AnimationBudgetAllocator.cpp
@@ -847,24 +847,29 @@ void FAnimationBudgetAllocator::OnHUDPostRender(AHUD* HUD, UCanvas* Canvas)
 }
 #endif
 
-void FAnimationBudgetAllocator::RemoveHelper(int32 Index)
+void FAnimationBudgetAllocator::RemoveHelper(int32 Index, USkeletalMeshComponentBudgeted* InExpectedComponent)
 {
-	if(AllComponentData.IsValidIndex(Index))
+	// Components are removed when they have been determined unreachable by GC, but may still attempt to unregister
+	// themselves redundantly during teardown, so make sure the index refers to the object we expect to remove
+	if (AllComponentData.IsValidIndex(Index))
 	{
-		if(AllComponentData[Index].Component != nullptr)
+		USkeletalMeshComponentBudgeted* const CurrentComponent = AllComponentData[Index].Component;
+		if (CurrentComponent == InExpectedComponent)
 		{
-			AllComponentData[Index].Component->SetAnimationBudgetHandle(INDEX_NONE);
-		}
+			if (CurrentComponent != nullptr)
+			{
+				CurrentComponent->SetAnimationBudgetHandle(INDEX_NONE);
+			}
 
-		AllComponentData.RemoveAtSwap(Index, 1, false);
+			AllComponentData.RemoveAtSwap(Index, 1, false);
 
-		// Update handle of swapped component
-		const int32 NumRemaining = AllComponentData.Num();
-		if(NumRemaining > 0 && Index != NumRemaining)
-		{
-			if(AllComponentData[Index].Component != nullptr)
+			// Update handle of swapped component
+			if (AllComponentData.IsValidIndex(Index))
 			{
-				AllComponentData[Index].Component->SetAnimationBudgetHandle(Index);
+				if (USkeletalMeshComponentBudgeted* SwappedComponent = AllComponentData[Index].Component)
+				{
+					SwappedComponent->SetAnimationBudgetHandle(Index);
+				}
 			}
 		}
 	}
@@ -931,9 +936,9 @@ void FAnimationBudgetAllocator::UnregisterComponent(USkeletalMeshComponentBudget
 	if (FAnimationBudgetAllocator::bCachedEnabled && bEnabled)
 	{
 		int32 ManagerHandle = InComponent->GetAnimationBudgetHandle();
-		if(ManagerHandle != INDEX_NONE)
+		if (ManagerHandle != INDEX_NONE)
 		{
-			RemoveHelper(ManagerHandle);
+			RemoveHelper(ManagerHandle, InComponent);
 
 			InComponent->bEnableUpdateRateOptimizations = true;
 			InComponent->EnableExternalTickRateControl(false);
@@ -969,22 +974,14 @@ void FAnimationBudgetAllocator::AddReferencedObjects(FReferenceCollector& Collec
 
 void FAnimationBudgetAllocator::HandlePostGarbageCollect()
 {
-	// Remove dead components, readjusting indices
-	bool bRemoved = false;
-	do
+	// Remove dead components backwards, readjusting indices
+	for (int32 DataIndex = AllComponentData.Num() - 1; DataIndex >= 0; --DataIndex)
 	{
-		bRemoved = false;
-		for(int32 DataIndex = 0; DataIndex < AllComponentData.Num(); ++DataIndex)
+		if (AllComponentData[DataIndex].Component == nullptr)
 		{
-			if(AllComponentData[DataIndex].Component == nullptr)
-			{
-				// We can remove while iterating here as we swap internally
-				RemoveHelper(DataIndex);
-				bRemoved = true;
-			}
+			RemoveHelper(DataIndex, nullptr);
 		}
 	}
-	while(bRemoved);
 }
 
 void FAnimationBudgetAllocator::SetGameThreadLastTickTimeMs(int32 InManagerHandle, float InGameThreadLastTickTimeMs)
diff --git a/Engine/Plugins/Runtime/AnimationBudgetAllocator/Source/AnimationBudgetAllocator/Private/AnimationBudgetAllocator.h b/Engine/Plugins/Runtime/AnimationBudgetAllocator/Source/AnimationBudgetAllocator/Private/AnimationBudgetAllocator.h
--- a/Engine/Plugins/Runtime/AnimationBudgetAllocator/Source/AnimationBudgetAllocator/Private/AnimationBudgetAllocator.h
+++ b/Engine/Plugins/Runtime/AnimationBudgetAllocator/Source/AnimationBudgetAllocator/Private/AnimationBudgetAllocator.h
@@ -175,7 +175,7 @@ protected:
 	int32 CalculateWorkDistributionAndQueue(float InDeltaSeconds, float& OutAverageTickRate);
 
 	/** Helper function for keeping handle indices in sync */
-	void RemoveHelper(int32 Index);
+	void RemoveHelper(int32 Index, USkeletalMeshComponentBudgeted* InExpectedComponent);
 
 	/** Helper function to enable/disable ticks */
 	void TickEnableHelper(USkeletalMeshComponent* InComponent, bool bInEnable);
diff --git a/Engine/Plugins/Runtime/AnimationBudgetAllocator/Source/AnimationBudgetAllocator/Private/SkeletalMeshComponentBudgeted.cpp b/Engine/Plugins/Runtime/AnimationBudgetAllocator/Source/AnimationBudgetAllocator/Private/SkeletalMeshComponentBudgeted.cpp
--- a/Engine/Plugins/Runtime/AnimationBudgetAllocator/Source/AnimationBudgetAllocator/Private/SkeletalMeshComponentBudgeted.cpp
+++ b/Engine/Plugins/Runtime/AnimationBudgetAllocator/Source/AnimationBudgetAllocator/Private/SkeletalMeshComponentBudgeted.cpp
@@ -37,16 +37,11 @@ void USkeletalMeshComponentBudgeted::BeginPlay()
 
 void USkeletalMeshComponentBudgeted::EndPlay(const EEndPlayReason::Type EndPlayReason)
 {
-	// Dont unregister if we are in the process of being destroyed in a GC.
-	// As reciprocal ptrs are null, handles are all invalid.
-	if(!IsUnreachable())	
+	if (UWorld* LocalWorld = GetWorld())
 	{
-		if (UWorld* LocalWorld = GetWorld())
+		if (IAnimationBudgetAllocator* LocalAnimationBudgetAllocator = IAnimationBudgetAllocator::Get(LocalWorld))
 		{
-			if (IAnimationBudgetAllocator* LocalAnimationBudgetAllocator = IAnimationBudgetAllocator::Get(LocalWorld))
-			{
-				LocalAnimationBudgetAllocator->UnregisterComponent(this);
-			}
+			LocalAnimationBudgetAllocator->UnregisterComponent(this);
 		}
 	}
 
