Animation budget allocator can now be used in BP-only projects

Made some tweaks to allow for the budget allocator to work when configured to be enabled in a level BP's BeginPlay call. This entailed adding a 'defered registration' system so components that get their BeginPlay called before the level BP end up properly registered.
Added optional auto-calculated significance so that projects do not always need to install their own callbacks.
Default min significance distance is configurable on a CVar (300m).

This also includes a fix for a crash that coudl occur if a component was destroyed as part of an animation notify calback (and was unregistered from the budget allocator inside the Super::CompleteParallelAnimationEvaluation call.

PR #9273: ????????? (Contributed by [REDACTED_PERSON])

diff --git a/Engine/Plugins/Runtime/AnimationBudgetAllocator/Source/AnimationBudgetAllocator/Private/AnimationBudgetAllocator.cpp b/Engine/Plugins/Runtime/AnimationBudgetAllocator/Source/AnimationBudgetAllocator/Private/AnimationBudgetAllocator.cpp
--- a/Engine/Plugins/Runtime/AnimationBudgetAllocator/Source/AnimationBudgetAllocator/Private/AnimationBudgetAllocator.cpp
+++ b/Engine/Plugins/Runtime/AnimationBudgetAllocator/Source/AnimationBudgetAllocator/Private/AnimationBudgetAllocator.cpp
@@ -72,8 +72,11 @@ FAnimationBudgetAllocator::FAnimationBudgetAllocator(UWorld* InWorld)
 	, CurrentFrameOffset(0)
 	, bEnabled(false)
 {
+	FAnimationBudgetAllocator::bCachedEnabled = GAnimationBudgetEnabled == 1;
+	
 	SetParametersFromCVars();
 
+	OnWorldBeginPlayHandle = InWorld->OnWorldBeginPlay.AddRaw(this, &FAnimationBudgetAllocator::HandleWorldBeginPlay);
 	PostGarbageCollectHandle = FCoreUObjectDelegates::GetPostGarbageCollect().AddRaw(this, &FAnimationBudgetAllocator::HandlePostGarbageCollect);
 	OnWorldPreActorTickHandle = FWorldDelegates::OnWorldPreActorTick.AddRaw(this, &FAnimationBudgetAllocator::OnWorldPreActorTick);
 	OnCVarParametersChangedHandle = GOnCVarParametersChanged.AddRaw(this, &FAnimationBudgetAllocator::SetParametersFromCVars);
@@ -87,6 +90,10 @@ FAnimationBudgetAllocator::~FAnimationBudgetAllocator()
 #if ENABLE_DRAW_DEBUG
 	AHUD::OnHUDPostRender.Remove(OnHUDPostRenderHandle);
 #endif
+	if(World)
+	{
+		World->OnWorldBeginPlay.Remove(OnWorldBeginPlayHandle);
+	}
 	FCoreUObjectDelegates::GetPostGarbageCollect().Remove(PostGarbageCollectHandle);
 	FWorldDelegates::OnWorldPreActorTick.Remove(OnWorldPreActorTickHandle);
 	GOnCVarParametersChanged.Remove(OnCVarParametersChangedHandle);
@@ -164,7 +171,9 @@ void FAnimationBudgetAllocator::QueueSortedComponentIndices(float InDeltaSeconds
 	TotalEstimatedTickTimeMs = 0.0f;
 	NumWorkUnitsForAverage = 0.0f;
 
-	auto QueueComponentTick = [InDeltaSeconds, this](FAnimBudgetAllocatorComponentData& InComponentData, int32 InComponentIndex, bool bInOnScreen)
+	const bool bUseSignificanceDelegate = USkeletalMeshComponentBudgeted::OnCalculateSignificance().IsBound();
+
+	auto QueueComponentTick = [InDeltaSeconds, this, bUseSignificanceDelegate](FAnimBudgetAllocatorComponentData& InComponentData, int32 InComponentIndex, bool bInOnScreen)
 	{
 		InComponentData.AccumulatedDeltaTime += InDeltaSeconds;
 		InComponentData.bOnScreen = bInOnScreen;
@@ -191,9 +200,16 @@ void FAnimationBudgetAllocator::QueueSortedComponentIndices(float InDeltaSeconds
 		// Auto-calculate significance here if we are set to
 		if(InComponentData.bAutoCalculateSignificance)
 		{
-			check(USkeletalMeshComponentBudgeted::OnCalculateSignificance().IsBound());
-
-			const float Significance = USkeletalMeshComponentBudgeted::OnCalculateSignificance().Execute(InComponentData.Component);
+			float Significance;
+			if(bUseSignificanceDelegate)
+			{
+				Significance = USkeletalMeshComponentBudgeted::OnCalculateSignificance().Execute(InComponentData.Component);
+			}
+			else
+			{
+				Significance = InComponentData.Component->GetDefaultSignificance();
+			}
+			
 			SetComponentSignificance(InComponentData.Component, Significance);
 		}
 	};
@@ -1060,4 +1076,23 @@ void FAnimationBudgetAllocator::SetParameters(const FAnimationBudgetAllocatorPar
 void FAnimationBudgetAllocator::SetParametersFromCVars()
 {
 	Parameters = GBudgetParameters;
+}
+
+void FAnimationBudgetAllocator::RegisterComponentDeferred(USkeletalMeshComponentBudgeted* InComponent)
+{
+	DeferredRegistrations.Add(InComponent); 
+}
+
+void FAnimationBudgetAllocator::HandleWorldBeginPlay()
+{
+	// This will catch worlds in (e.g.) PIE that try to set the CVar on startup
+	FAnimationBudgetAllocator::bCachedEnabled = GAnimationBudgetEnabled == 1;
+
+	// Run thru all deferred registrations
+	for(USkeletalMeshComponentBudgeted* Component : DeferredRegistrations)
+	{
+		RegisterComponent(Component);
+	}
+
+	DeferredRegistrations.Empty();
 }
\ No newline at end of file
diff --git a/Engine/Plugins/Runtime/AnimationBudgetAllocator/Source/AnimationBudgetAllocator/Private/AnimationBudgetAllocator.h b/Engine/Plugins/Runtime/AnimationBudgetAllocator/Source/AnimationBudgetAllocator/Private/AnimationBudgetAllocator.h
--- a/Engine/Plugins/Runtime/AnimationBudgetAllocator/Source/AnimationBudgetAllocator/Private/AnimationBudgetAllocator.h
+++ b/Engine/Plugins/Runtime/AnimationBudgetAllocator/Source/AnimationBudgetAllocator/Private/AnimationBudgetAllocator.h
@@ -158,13 +158,19 @@ public:
 		return TEXT("FAnimationBudgetAllocator");
 	}
 
+	// Registers a component for deferred registration (post begin play)
+	void RegisterComponentDeferred(USkeletalMeshComponentBudgeted* InComponent);
+	
 protected:
 	/** We tick before all actors in the world using this delegate */
 	void OnWorldPreActorTick(UWorld* InWorld, ELevelTick InLevelTick, float InDeltaSeconds);
 
 	/** Clean up dead components post GC */
 	void HandlePostGarbageCollect();
- 
+
+	/** Handle when a world begins play */
+	void HandleWorldBeginPlay();
+	
 	/** First pass of the Update(). Queues component indices that want to tick into AllSortedComponentData. */
 	void QueueSortedComponentIndices(float InDeltaSeconds);
 
@@ -218,6 +224,9 @@ protected:
 	/** All non-rendered components we might tick */
 	TArray<int32> NonRenderedComponentData;
 
+	/** Component registrations for our world before beginplay was called */
+	TArray<USkeletalMeshComponentBudgeted*> DeferredRegistrations;
+	
 #if ENABLE_DRAW_DEBUG
 	/** Recorded debug times */
 	TArray<FVector2D> DebugTimes;
@@ -268,6 +277,9 @@ protected:
 	/** Handle used for debug drawing */
 	FDelegateHandle OnHUDPostRenderHandle;
 
+	/** Handle used to handle world begin play */
+	FDelegateHandle OnWorldBeginPlayHandle;
+	
 	/** Offset used to distribute component ticks */
 	uint32 CurrentFrameOffset;
 
diff --git a/Engine/Plugins/Runtime/AnimationBudgetAllocator/Source/AnimationBudgetAllocator/Private/AnimationBudgetAllocatorCVars.cpp b/Engine/Plugins/Runtime/AnimationBudgetAllocator/Source/AnimationBudgetAllocator/Private/AnimationBudgetAllocatorCVars.cpp
--- a/Engine/Plugins/Runtime/AnimationBudgetAllocator/Source/AnimationBudgetAllocator/Private/AnimationBudgetAllocatorCVars.cpp
+++ b/Engine/Plugins/Runtime/AnimationBudgetAllocator/Source/AnimationBudgetAllocator/Private/AnimationBudgetAllocatorCVars.cpp
@@ -331,4 +331,17 @@ static FAutoConsoleVariableRef CVarSkelBatch_BudgetPressureBeforeEmergencyReduce
 		GBudgetParameters.BudgetPressureBeforeEmergencyReducedWork = FMath::Max(GBudgetParameters.BudgetPressureBeforeEmergencyReducedWork, 0.0f);
 		GOnCVarParametersChanged.Broadcast();
 	}),
+	ECVF_Scalability);
+
+static FAutoConsoleVariableRef CVarSkelBatch_AutoCalculatedSignificanceMaxDistance(
+	TEXT("a.Budget.AutoCalculatedSignificanceMaxDistance"),
+	GBudgetParameters.AutoCalculatedSignificanceMaxDistance,
+	TEXT("Range > 1.0, Default = 300.0\n")
+	TEXT("Controls the distance (in cm) at which auto-calculated significance for budgeted components bottoms out. Components within the distance 1 -> Max will have significance mapped 1 -> 0, outside of MaxDistance significance will be zero.\n"),
+	FConsoleVariableDelegate::CreateLambda([](IConsoleVariable* InVariable)
+	{
+		GBudgetParameters.AutoCalculatedSignificanceMaxDistance = FMath::Max(GBudgetParameters.AutoCalculatedSignificanceMaxDistance, 1.0f);
+		GBudgetParameters.AutoCalculatedSignificanceMaxDistanceSqr = GBudgetParameters.AutoCalculatedSignificanceMaxDistance * GBudgetParameters.AutoCalculatedSignificanceMaxDistance;
+		GOnCVarParametersChanged.Broadcast();
+	}),
 	ECVF_Scalability);
\ No newline at end of file
diff --git a/Engine/Plugins/Runtime/AnimationBudgetAllocator/Source/AnimationBudgetAllocator/Private/SkeletalMeshComponentBudgeted.cpp b/Engine/Plugins/Runtime/AnimationBudgetAllocator/Source/AnimationBudgetAllocator/Private/SkeletalMeshComponentBudgeted.cpp
--- a/Engine/Plugins/Runtime/AnimationBudgetAllocator/Source/AnimationBudgetAllocator/Private/SkeletalMeshComponentBudgeted.cpp
+++ b/Engine/Plugins/Runtime/AnimationBudgetAllocator/Source/AnimationBudgetAllocator/Private/SkeletalMeshComponentBudgeted.cpp
@@ -2,8 +2,11 @@
 
 #include "SkeletalMeshComponentBudgeted.h"
 #include "AnimationBudgetAllocator.h"
+#include "AnimationBudgetAllocatorCVars.h"
+#include "SkeletalRenderPublic.h"
 #include "Kismet/KismetSystemLibrary.h"
 #include "ProfilingDebugging/CsvProfiler.h"
+#include "GameFramework/Controller.h"
 
 CSV_DECLARE_CATEGORY_EXTERN(AnimationBudget);
 
@@ -11,8 +14,6 @@ FOnCalculateSignificance USkeletalMeshComponentBudgeted::OnCalculateSignificance
 
 USkeletalMeshComponentBudgeted::USkeletalMeshComponentBudgeted(const FObjectInitializer& ObjectInitializer)
 	: Super(ObjectInitializer)
-	, AnimationBudgetHandle(INDEX_NONE)
-	, AnimationBudgetAllocator(nullptr)
 	, bAutoRegisterWithBudgetAllocator(true)
 	, bAutoCalculateSignificance(false)
 	, bShouldUseActorRenderedFlag(false)
@@ -27,9 +28,18 @@ void USkeletalMeshComponentBudgeted::BeginPlay()
 	{
 		if (UWorld* LocalWorld = GetWorld())
 		{
-			if (IAnimationBudgetAllocator* LocalAnimationBudgetAllocator = IAnimationBudgetAllocator::Get(LocalWorld))
+			if (FAnimationBudgetAllocator* LocalAnimationBudgetAllocator = static_cast<FAnimationBudgetAllocator*>(IAnimationBudgetAllocator::Get(LocalWorld)))
 			{
-				LocalAnimationBudgetAllocator->RegisterComponent(this);
+				if(LocalWorld->HasBegunPlay())
+				{
+					// World is playing, so register 
+					LocalAnimationBudgetAllocator->RegisterComponent(this);
+				}
+				else
+				{
+					// World is not playing, so register deferred 
+					LocalAnimationBudgetAllocator->RegisterComponentDeferred(this);
+				}
 			}
 		}
 	}
@@ -84,7 +94,32 @@ void USkeletalMeshComponentBudgeted::TickComponent(float DeltaTime, enum ELevelT
 
 		Super::TickComponent(DeltaTime, TickType, ThisTickFunction);
 
-		AnimationBudgetAllocator->SetGameThreadLastTickTimeMs(AnimationBudgetHandle, FPlatformTime::ToMilliseconds64(FPlatformTime::Cycles64() - StartTime));
+		if(bAutoCalculateSignificance && !USkeletalMeshComponentBudgeted::OnCalculateSignificance().IsBound())
+		{
+			// Default auto-calculated significance is based on distance to each player viewpoint
+			AutoCalculatedSignificance = 0.0f;
+
+			const FVector WorldLocation = GetComponentTransform().GetLocation();
+			const UWorld* ComponentWorld = GetWorld();
+			for (FConstControllerIterator It = ComponentWorld->GetControllerIterator(); It; ++It)
+			{
+				if(const AController* PlayerController = It->Get())
+				{
+					FVector PlayerViewLocation = FVector::ZeroVector;
+					FRotator PlayerViewRotation = FRotator::ZeroRotator;
+					PlayerController->GetPlayerViewPoint(PlayerViewLocation, PlayerViewRotation);
+
+					const float DistanceSqr = (WorldLocation - PlayerViewLocation).SizeSquared();
+					const float Significance = FMath::Max(GBudgetParameters.AutoCalculatedSignificanceMaxDistanceSqr - DistanceSqr, 1.0f) / GBudgetParameters.AutoCalculatedSignificanceMaxDistanceSqr;
+					AutoCalculatedSignificance = FMath::Max(AutoCalculatedSignificance, Significance);
+				}
+			}
+		}
+		
+		if(AnimationBudgetAllocator)
+		{
+			AnimationBudgetAllocator->SetGameThreadLastTickTimeMs(AnimationBudgetHandle, FPlatformTime::ToMilliseconds64(FPlatformTime::Cycles64() - StartTime));
+		}
 	}
 	else
 	{
@@ -104,10 +139,18 @@ void USkeletalMeshComponentBudgeted::CompleteParallelAnimationEvaluation(bool bD
 
 		Super::CompleteParallelAnimationEvaluation(bDoPostAnimEvaluation);
 
-		AnimationBudgetAllocator->SetGameThreadLastCompletionTimeMs(AnimationBudgetHandle, FPlatformTime::ToMilliseconds64(FPlatformTime::Cycles64() - StartTime));
+		if(AnimationBudgetAllocator)
+		{
+			AnimationBudgetAllocator->SetGameThreadLastCompletionTimeMs(AnimationBudgetHandle, FPlatformTime::ToMilliseconds64(FPlatformTime::Cycles64() - StartTime));
+		}
 	}
 	else
 	{
 		Super::CompleteParallelAnimationEvaluation(bDoPostAnimEvaluation);
 	}
+}
+
+float USkeletalMeshComponentBudgeted::GetDefaultSignificance() const
+{
+	return AutoCalculatedSignificance;
 }
\ No newline at end of file
diff --git a/Engine/Plugins/Runtime/AnimationBudgetAllocator/Source/AnimationBudgetAllocator/Public/AnimationBudgetAllocatorParameters.h b/Engine/Plugins/Runtime/AnimationBudgetAllocator/Source/AnimationBudgetAllocator/Public/AnimationBudgetAllocatorParameters.h
--- a/Engine/Plugins/Runtime/AnimationBudgetAllocator/Source/AnimationBudgetAllocator/Public/AnimationBudgetAllocatorParameters.h
+++ b/Engine/Plugins/Runtime/AnimationBudgetAllocator/Source/AnimationBudgetAllocator/Public/AnimationBudgetAllocatorParameters.h
@@ -155,4 +155,13 @@ struct FAnimationBudgetAllocatorParameters
 	 */
 	UPROPERTY(BlueprintReadWrite, Category = "Parameters")
 	float BudgetPressureBeforeEmergencyReducedWork = 2.5f;
+
+	/**
+	 * Range > 1.0.
+	 * Controls the distance at which auto-calculated significance for budgeted components bottoms out. Components
+	 * within the distance 1 -> Max will have significance mapped 1 -> 0, outside of MaxDistance significance will be zero.
+	 */
+	UPROPERTY(BlueprintReadWrite, Category = "Parameters")
+	float AutoCalculatedSignificanceMaxDistance = 30000.0f;
+	float AutoCalculatedSignificanceMaxDistanceSqr = 30000.0f * 30000.0f;
 };
\ No newline at end of file
diff --git a/Engine/Plugins/Runtime/AnimationBudgetAllocator/Source/AnimationBudgetAllocator/Public/SkeletalMeshComponentBudgeted.h b/Engine/Plugins/Runtime/AnimationBudgetAllocator/Source/AnimationBudgetAllocator/Public/SkeletalMeshComponentBudgeted.h
--- a/Engine/Plugins/Runtime/AnimationBudgetAllocator/Source/AnimationBudgetAllocator/Public/SkeletalMeshComponentBudgeted.h
+++ b/Engine/Plugins/Runtime/AnimationBudgetAllocator/Source/AnimationBudgetAllocator/Public/SkeletalMeshComponentBudgeted.h
@@ -72,27 +72,37 @@ private:
 	/** Set the budget allocator that is tracking us */
 	void SetAnimationBudgetAllocator(IAnimationBudgetAllocator* InAnimationBudgetAllocator) { AnimationBudgetAllocator = InAnimationBudgetAllocator; }
 
+	/** Handles self-registration after actors have begun play */
+	void HandleWorldBeginPlay();
+
+	/** Gets the default significance value based on distance */
+	float GetDefaultSignificance() const;
+	
 private:
 	/** Delegate called to increase/decrease the amount of work a component performs */
 	FOnReduceWork OnReduceWorkDelegate;
 
 	/** Delegate called to calculate significance if bAutoCalculateSignificance = true */
 	static FOnCalculateSignificance OnCalculateSignificanceDelegate;
+	
+	/** Owning animation budget allocator */
+	IAnimationBudgetAllocator* AnimationBudgetAllocator = nullptr;
 
 	/** Handle used for identification */
-	int32 AnimationBudgetHandle;
-
-	/** Owning animation budget allocator */
-	IAnimationBudgetAllocator* AnimationBudgetAllocator;
+	int32 AnimationBudgetHandle = INDEX_NONE;
 
+	/** Cached calculated significance */
+	float AutoCalculatedSignificance = 1.0f;
+	
 	/** Whether this component should automatically register with the budget allocator in OnRegister/OnUnregister */
 	UPROPERTY(EditAnywhere, BlueprintSetter=SetAutoRegisterWithBudgetAllocator, Category = Budgeting)
 	uint8 bAutoRegisterWithBudgetAllocator : 1;
 
-	/** Whether this component should automatically register with the budget allocator in OnRegister/OnUnregister */
+	/** Whether this component should automatically calculate its significance (rather than some external system pushing the significance to it) */
 	UPROPERTY(EditAnywhere, Category = Budgeting)
 	uint8 bAutoCalculateSignificance : 1;
 
-	UPROPERTY(EditAnywhere, Category = Budgeting)
+	/** Whether this component should use its owning actor's rendered state to determine visibility. If this is not set then the component's visibility will be used */
+	UPROPERTY(EditAnywhere, Category = Budgeting, AdvancedDisplay)
 	uint8 bShouldUseActorRenderedFlag : 1;
 };
\ No newline at end of file
