fix animation budgeter interpolation

diff --git a/Engine/Plugins/Runtime/AnimationBudgetAllocator/Source/AnimationBudgetAllocator/Private/AnimationBudgetAllocator.cpp b/Engine/Plugins/Runtime/AnimationBudgetAllocator/Source/AnimationBudgetAllocator/Private/AnimationBudgetAllocator.cpp
--- a/Engine/Plugins/Runtime/AnimationBudgetAllocator/Source/AnimationBudgetAllocator/Private/AnimationBudgetAllocator.cpp
+++ b/Engine/Plugins/Runtime/AnimationBudgetAllocator/Source/AnimationBudgetAllocator/Private/AnimationBudgetAllocator.cpp
@@ -365,18 +365,21 @@ int32 FAnimationBudgetAllocator::CalculateWorkDistributionAndQueue(float InDelta
 		const int32 PrerequisiteHandle = InComponentData.RootPrerequisite != nullptr ? InComponentData.RootPrerequisite->GetAnimationBudgetHandle() : INDEX_NONE;
 		const FAnimBudgetAllocatorComponentData& ComponentDataToCheck = PrerequisiteHandle != INDEX_NONE ? AllComponentData[PrerequisiteHandle] : InComponentData;
 
+		// Using (frame offset + frame counter) % tick rate allows us to only tick at the specified interval,
+		// but at a roughly even distribution over all registered components
+		const bool bTickThisFrame = (((GFrameCounter + ComponentDataToCheck.FrameOffset) % ComponentDataToCheck.TickRate) == 0);
 		// Calculate interp alpha even when interpolation or ticking is disabled as this is used to decide how much root motion to consume
 		// each frame by character movement
-		float Alpha = (InComponentData.TickRate > 1) ? FMath::Clamp((1.0f / (InComponentData.TickRate - InComponentData.SkippedTicks + 1)), 0.0f, 1.0f) : 1.f;
+		InComponentData.SkippedTicks = bTickThisFrame ? 0 : (InComponentData.SkippedTicks + 1);
+		
+		// Each frame we replace the source pose with the interpolated pose, so Alpha is set to blend from the current pose, one tick closer to the final pose.
+		const int TicksRemaining = InComponentData.TickRate - InComponentData.SkippedTicks;
+		float Alpha = (InComponentData.TickRate > 1) ? FMath::Clamp((1.0f / TicksRemaining), 0.0f, 1.0f) : 1.f;
 		InComponentData.Component->SetExternalInterpolationAlpha(Alpha);
 
-		// Using (frame offset + frame counter) % tick rate allows us to only tick at the specified interval,
-		// but at a roughly even distribution over all registered components
-		const bool bTickThisFrame = (((GFrameCounter + ComponentDataToCheck.FrameOffset) % ComponentDataToCheck.TickRate) == 0);
 		if((ComponentDataToCheck.bInterpolate && ComponentDataToCheck.bOnScreen) || bTickThisFrame)
 		{
 			InComponentData.bInterpolate = ComponentDataToCheck.bInterpolate;
-			InComponentData.SkippedTicks = bTickThisFrame ? 0 : (InComponentData.SkippedTicks + 1);
 
 			// Reset completion time as it may not always be run
 			InComponentData.GameThreadLastCompletionTimeMs = 0.0f;
diff --git a/Engine/Source/Runtime/Engine/Classes/Animation/AnimationAsset.h b/Engine/Source/Runtime/Engine/Classes/Animation/AnimationAsset.h
--- a/Engine/Source/Runtime/Engine/Classes/Animation/AnimationAsset.h
+++ b/Engine/Source/Runtime/Engine/Classes/Animation/AnimationAsset.h
@@ -783,12 +783,10 @@ public:
 
 	FRootMotionMovementParams ConsumeRootMotion(float Alpha)
 	{
-		const ScalarRegister VAlpha(Alpha);
-		FTransform PartialRootMotion = (RootMotionTransform*VAlpha);
-		PartialRootMotion.SetScale3D(RootMotionScale); // Reset scale after multiplication above
-		PartialRootMotion.NormalizeRotation();
+		FTransform PartialRootMotion(FQuat::Slerp(FQuat::Identity, RootMotionTransform.GetRotation(), Alpha), RootMotionTransform.GetTranslation()*Alpha, RootMotionScale);
+
+		// remove the part of the root motion we are applying now and leave the remainder in RootMotionTransform
 		RootMotionTransform = RootMotionTransform.GetRelativeTransform(PartialRootMotion);
-		RootMotionTransform.NormalizeRotation(); //Make sure we are normalized, this needs to be investigated further
 
 		FRootMotionMovementParams ReturnParams;
 		ReturnParams.Set(PartialRootMotion);
