Only look at player controllers (not e.g. AI controllers) when considering viewpoints for auto significance calculation

#10074

diff --git a/Engine/Plugins/Runtime/AnimationBudgetAllocator/Source/AnimationBudgetAllocator/Private/SkeletalMeshComponentBudgeted.cpp b/Engine/Plugins/Runtime/AnimationBudgetAllocator/Source/AnimationBudgetAllocator/Private/SkeletalMeshComponentBudgeted.cpp
--- a/Engine/Plugins/Runtime/AnimationBudgetAllocator/Source/AnimationBudgetAllocator/Private/SkeletalMeshComponentBudgeted.cpp
+++ b/Engine/Plugins/Runtime/AnimationBudgetAllocator/Source/AnimationBudgetAllocator/Private/SkeletalMeshComponentBudgeted.cpp
@@ -103,9 +103,9 @@ void USkeletalMeshComponentBudgeted::TickComponent(float DeltaTime, enum ELevelT
 
 			const FVector WorldLocation = GetComponentTransform().GetLocation();
 			const UWorld* ComponentWorld = GetWorld();
-			for (FConstControllerIterator It = ComponentWorld->GetControllerIterator(); It; ++It)
+			for (FConstPlayerControllerIterator It = ComponentWorld->GetPlayerControllerIterator(); It; ++It)
 			{
-				if(const AController* PlayerController = It->Get())
+				if(const APlayerController* PlayerController = It->Get())
 				{
 					FVector PlayerViewLocation = FVector::ZeroVector;
 					FRotator PlayerViewRotation = FRotator::ZeroRotator;
