PR #10467: Fix crash on seamless travel with Animation Budgeter

diff --git a/Engine/Plugins/Runtime/AnimationBudgetAllocator/Source/AnimationBudgetAllocator/Private/AnimationBudgetAllocatorModule.cpp b/Engine/Plugins/Runtime/AnimationBudgetAllocator/Source/AnimationBudgetAllocator/Private/AnimationBudgetAllocatorModule.cpp
--- a/Engine/Plugins/Runtime/AnimationBudgetAllocator/Source/AnimationBudgetAllocator/Private/AnimationBudgetAllocatorModule.cpp
+++ b/Engine/Plugins/Runtime/AnimationBudgetAllocator/Source/AnimationBudgetAllocator/Private/AnimationBudgetAllocatorModule.cpp
@@ -49,6 +49,9 @@ void FAnimationBudgetAllocatorModule::HandleWorldCleanup(UWorld* World, bool bSe
 	FAnimationBudgetAllocator* Budgeter = WorldAnimationBudgetAllocators.FindRef(World);
 	if(Budgeter)
 	{
+		//Cleanup components, as this can cause a crash in some conditions during seamless travel as actors can be pulled in
+		//to the travel world.
+		Budgeter->SetEnabled(false);
 		delete Budgeter;
 		WorldAnimationBudgetAllocators.Remove(World);
 	}
