#UE Fix static assert to only apply to non-editor builds since editoronly properties bump the size higher than this enforces

Original CL Desc
-----------------------------------------------------------------
Improving packing on SkeletalMeshComponent which should get a 16 byte gain. This will cause SkeletalMeshComponentBudgeted to drop to the 4080 bin for all platforms other than Switch, getting 288 bytes back for each.

diff --git a/Engine/Plugins/Runtime/AnimationBudgetAllocator/Source/AnimationBudgetAllocator/Private/SkeletalMeshComponentBudgeted.cpp b/Engine/Plugins/Runtime/AnimationBudgetAllocator/Source/AnimationBudgetAllocator/Private/SkeletalMeshComponentBudgeted.cpp
--- a/Engine/Plugins/Runtime/AnimationBudgetAllocator/Source/AnimationBudgetAllocator/Private/SkeletalMeshComponentBudgeted.cpp
+++ b/Engine/Plugins/Runtime/AnimationBudgetAllocator/Source/AnimationBudgetAllocator/Private/SkeletalMeshComponentBudgeted.cpp
@@ -12,7 +12,7 @@
 
 CSV_DECLARE_CATEGORY_EXTERN(AnimationBudget);
 
-#if UE_BUILD_SHIPPING
+#if UE_BUILD_SHIPPING && !WITH_EDITOR
 static_assert(sizeof(USkeletalMeshComponentBudgeted) <= 4080, "We expect SkeletalMeshComponentBudgeted to be 4080 or less, otherwise we waste memory on 288bytes of padding on MallocBinned2. Please improve padding on this class hierarchy to pay for your added variable.");
 #endif
 
