Fix to AnimationBudgeAllocator that keeps tick rates consistent between tick prerequisites

diff --git a/Engine/Plugins/Runtime/AnimationBudgetAllocator/Source/AnimationBudgetAllocator/Private/AnimationBudgetAllocator.cpp b/Engine/Plugins/Runtime/AnimationBudgetAllocator/Source/AnimationBudgetAllocator/Private/AnimationBudgetAllocator.cpp
--- a/Engine/Plugins/Runtime/AnimationBudgetAllocator/Source/AnimationBudgetAllocator/Private/AnimationBudgetAllocator.cpp
+++ b/Engine/Plugins/Runtime/AnimationBudgetAllocator/Source/AnimationBudgetAllocator/Private/AnimationBudgetAllocator.cpp
@@ -554,16 +554,16 @@ int32 FAnimationBudgetAllocator::CalculateWorkDistributionAndQueue(float InDelta
 
 			SET_FLOAT_STAT(STAT_AnimationBudgetAllocator_SmoothedBudgetPressure, SmoothedBudgetPressure);
 
-			// Queue for tick
+			// Check all datas to line up tick rates based on prerequisites before queuing for ticks so ticks can be consistent
 			for (SortedComponentIndex = 0; SortedComponentIndex < TotalIdealWorkUnits; ++SortedComponentIndex)
 			{
 				FAnimBudgetAllocatorComponentData& ComponentData = AllComponentData[AllSortedComponentData[SortedComponentIndex]];
 
 				// Ensure that root prerequisite doesnt end up with a lower (or different) tick rate than dependencies
-				if(ComponentData.RootPrerequisite != nullptr)
+				if (ComponentData.RootPrerequisite != nullptr)
 				{
 					const int32 PrerequisiteHandle = ComponentData.RootPrerequisite->GetAnimationBudgetHandle();
-					if(PrerequisiteHandle != INDEX_NONE)
+					if (PrerequisiteHandle != INDEX_NONE)
 					{
 						FAnimBudgetAllocatorComponentData& RootPrerequisiteComponentData = AllComponentData[PrerequisiteHandle];
 						RootPrerequisiteComponentData.TickRate = ComponentData.TickRate = FMath::Min(ComponentData.TickRate, RootPrerequisiteComponentData.TickRate);
@@ -571,6 +571,12 @@ int32 FAnimationBudgetAllocator::CalculateWorkDistributionAndQueue(float InDelta
 						RootPrerequisiteComponentData.StateChangeThrottle = ComponentData.StateChangeThrottle = FMath::Min(ComponentData.StateChangeThrottle, RootPrerequisiteComponentData.StateChangeThrottle);
 					}
 				}
+			}
+
+			// Queue for tick
+			for (SortedComponentIndex = 0; SortedComponentIndex < TotalIdealWorkUnits; ++SortedComponentIndex)
+			{
+				FAnimBudgetAllocatorComponentData& ComponentData = AllComponentData[AllSortedComponentData[SortedComponentIndex]];
 
 				QueueForTick(ComponentData, StateChangeThrottleInFrames);
 			}
