Add ForceNextTickThisFrame function to AnimationBudgetAllocator.

diff --git a/Engine/Plugins/Runtime/AnimationBudgetAllocator/Source/AnimationBudgetAllocator/Private/AnimationBudgetAllocator.cpp b/Engine/Plugins/Runtime/AnimationBudgetAllocator/Source/AnimationBudgetAllocator/Private/AnimationBudgetAllocator.cpp
--- a/Engine/Plugins/Runtime/AnimationBudgetAllocator/Source/AnimationBudgetAllocator/Private/AnimationBudgetAllocator.cpp
+++ b/Engine/Plugins/Runtime/AnimationBudgetAllocator/Source/AnimationBudgetAllocator/Private/AnimationBudgetAllocator.cpp
@@ -1077,6 +1077,21 @@ void FAnimationBudgetAllocator::SetParameters(const FAnimationBudgetAllocatorPar
 	Parameters = InParameters;
 }
 
+void FAnimationBudgetAllocator::ForceNextTickThisFrame(USkeletalMeshComponentBudgeted* InComponent)
+{
+	if (GAnimationBudgetEnabled && bEnabled)
+	{
+		int32 ManagerHandle = InComponent->GetAnimationBudgetHandle();
+		if (ManagerHandle != INDEX_NONE)
+		{
+			FAnimBudgetAllocatorComponentData& ComponentData = AllComponentData[ManagerHandle];
+
+			// set frameoffset so next tick will be on this frame
+			ComponentData.FrameOffset += (ComponentData.TickRate - ((GFrameCounter + ComponentData.FrameOffset) % ComponentData.TickRate));
+		}
+	}
+}
+
 void FAnimationBudgetAllocator::SetParametersFromCVars()
 {
 	Parameters = GBudgetParameters;
diff --git a/Engine/Plugins/Runtime/AnimationBudgetAllocator/Source/AnimationBudgetAllocator/Private/AnimationBudgetAllocator.h b/Engine/Plugins/Runtime/AnimationBudgetAllocator/Source/AnimationBudgetAllocator/Private/AnimationBudgetAllocator.h
--- a/Engine/Plugins/Runtime/AnimationBudgetAllocator/Source/AnimationBudgetAllocator/Private/AnimationBudgetAllocator.h
+++ b/Engine/Plugins/Runtime/AnimationBudgetAllocator/Source/AnimationBudgetAllocator/Private/AnimationBudgetAllocator.h
@@ -144,6 +144,7 @@ public:
 	virtual void SetEnabled(bool bInEnabled) override;
 	virtual bool GetEnabled() const override;
 	virtual void SetParameters(const FAnimationBudgetAllocatorParameters& InParameters) override;
+	virtual void ForceNextTickThisFrame(USkeletalMeshComponentBudgeted* Component) override;
 
 	// FGCObject interface
 	virtual void AddReferencedObjects(FReferenceCollector& Collector) override;
diff --git a/Engine/Plugins/Runtime/AnimationBudgetAllocator/Source/AnimationBudgetAllocator/Public/IAnimationBudgetAllocator.h b/Engine/Plugins/Runtime/AnimationBudgetAllocator/Source/AnimationBudgetAllocator/Public/IAnimationBudgetAllocator.h
--- a/Engine/Plugins/Runtime/AnimationBudgetAllocator/Source/AnimationBudgetAllocator/Public/IAnimationBudgetAllocator.h
+++ b/Engine/Plugins/Runtime/AnimationBudgetAllocator/Source/AnimationBudgetAllocator/Public/IAnimationBudgetAllocator.h
@@ -72,4 +72,7 @@ public:
 
 	/** Set the various parameters */
 	virtual void SetParameters(const FAnimationBudgetAllocatorParameters& InParameters) = 0;
+
+	/** Force component to tick this frame regardless of tick rate */
+	virtual void ForceNextTickThisFrame(USkeletalMeshComponentBudgeted* Component) = 0;
 };
\ No newline at end of file
