Fix budgeted meshes that are placed in levels not getting registered when playing as client

Stop relying on the worlds HasBegunPlay, as that is inconsistent with the calls to OnBeginPlay.
Remove ordering issues with enabling the system (e.g. via a level script) and component registration by always registering/unregistering regardless of enabled state.

diff --git a/Engine/Plugins/Runtime/AnimationBudgetAllocator/Source/AnimationBudgetAllocator/Private/AnimationBudgetAllocator.cpp b/Engine/Plugins/Runtime/AnimationBudgetAllocator/Source/AnimationBudgetAllocator/Private/AnimationBudgetAllocator.cpp
--- a/Engine/Plugins/Runtime/AnimationBudgetAllocator/Source/AnimationBudgetAllocator/Private/AnimationBudgetAllocator.cpp
+++ b/Engine/Plugins/Runtime/AnimationBudgetAllocator/Source/AnimationBudgetAllocator/Private/AnimationBudgetAllocator.cpp
@@ -66,6 +66,7 @@ FAnimationBudgetAllocator::FAnimationBudgetAllocator(UWorld* InWorld)
 	, ReducedComponentWorkCounter(0)
 	, CurrentFrameOffset(0)
 	, bEnabled(false)
+	, bHasBegunPlay(InWorld->HasBegunPlay())
 {
 	FAnimationBudgetAllocator::bCachedEnabled = GAnimationBudgetEnabled == 1;
 	
@@ -927,57 +928,48 @@ static USkeletalMeshComponentBudgeted* FindRootPrerequisite(USkeletalMeshCompone
 
 void FAnimationBudgetAllocator::RegisterComponent(USkeletalMeshComponentBudgeted* InComponent)
 {
-	if (FAnimationBudgetAllocator::bCachedEnabled && bEnabled)
+	if (InComponent->GetAnimationBudgetHandle() == INDEX_NONE)
 	{
-		if (InComponent->GetAnimationBudgetHandle() == INDEX_NONE)
-		{
-			InComponent->bEnableUpdateRateOptimizations = false;
-			InComponent->EnableExternalTickRateControl(true);
-			InComponent->SetAnimationBudgetHandle(AllComponentData.Num());
-
-			// Setup frame offset
-			FAnimBudgetAllocatorComponentData& ComponentData = AllComponentData.Emplace_GetRef(InComponent, Parameters.InitialEstimatedWorkUnitTimeMs, Parameters.StateChangeThrottleInFrames);
-			USkeletalMeshComponentBudgeted* RootPrerequisite = FindRootPrerequisite(InComponent);
-			ComponentData.RootPrerequisite = (RootPrerequisite != nullptr && RootPrerequisite != InComponent) ? RootPrerequisite : nullptr;
-			ComponentData.FrameOffset = CurrentFrameOffset++;
-			ComponentData.bAutoCalculateSignificance = InComponent->GetAutoCalculateSignificance();
-
-			InComponent->SetAnimationBudgetAllocator(this);
-		}
-		else
-		{
-			UpdateComponentTickPrerequsites(InComponent);
-		}
+		InComponent->bEnableUpdateRateOptimizations = false;
+		InComponent->EnableExternalTickRateControl(true);
+		InComponent->SetAnimationBudgetHandle(AllComponentData.Num());
+
+		// Setup frame offset
+		FAnimBudgetAllocatorComponentData& ComponentData = AllComponentData.Emplace_GetRef(InComponent, Parameters.InitialEstimatedWorkUnitTimeMs, Parameters.StateChangeThrottleInFrames);
+		USkeletalMeshComponentBudgeted* RootPrerequisite = FindRootPrerequisite(InComponent);
+		ComponentData.RootPrerequisite = (RootPrerequisite != nullptr && RootPrerequisite != InComponent) ? RootPrerequisite : nullptr;
+		ComponentData.FrameOffset = CurrentFrameOffset++;
+		ComponentData.bAutoCalculateSignificance = InComponent->GetAutoCalculateSignificance();
+
+		InComponent->SetAnimationBudgetAllocator(this);
+	}
+	else
+	{
+		UpdateComponentTickPrerequsites(InComponent);
 	}
 }
 
 void FAnimationBudgetAllocator::UnregisterComponent(USkeletalMeshComponentBudgeted* InComponent)
 {
-	if (FAnimationBudgetAllocator::bCachedEnabled && bEnabled)
+	int32 ManagerHandle = InComponent->GetAnimationBudgetHandle();
+	if (ManagerHandle != INDEX_NONE)
 	{
-		int32 ManagerHandle = InComponent->GetAnimationBudgetHandle();
-		if (ManagerHandle != INDEX_NONE)
-		{
-			RemoveHelper(ManagerHandle, InComponent);
+		RemoveHelper(ManagerHandle, InComponent);
 
-			InComponent->bEnableUpdateRateOptimizations = true;
-			InComponent->EnableExternalTickRateControl(false);
-			InComponent->SetAnimationBudgetAllocator(nullptr);
-		}
+		InComponent->bEnableUpdateRateOptimizations = true;
+		InComponent->EnableExternalTickRateControl(false);
+		InComponent->SetAnimationBudgetAllocator(nullptr);
 	}
 }
 
 void FAnimationBudgetAllocator::UpdateComponentTickPrerequsites(USkeletalMeshComponentBudgeted* InComponent)
 {
-	if (FAnimationBudgetAllocator::bCachedEnabled && bEnabled)
+	int32 ManagerHandle = InComponent->GetAnimationBudgetHandle();
+	if(ManagerHandle != INDEX_NONE)
 	{
-		int32 ManagerHandle = InComponent->GetAnimationBudgetHandle();
-		if(ManagerHandle != INDEX_NONE)
-		{
-			FAnimBudgetAllocatorComponentData& ComponentData = AllComponentData[ManagerHandle];
-			USkeletalMeshComponentBudgeted* RootPrerequisite = FindRootPrerequisite(InComponent);
-			ComponentData.RootPrerequisite = (RootPrerequisite != nullptr && RootPrerequisite != InComponent) ? RootPrerequisite : nullptr;
-		}
+		FAnimBudgetAllocatorComponentData& ComponentData = AllComponentData[ManagerHandle];
+		USkeletalMeshComponentBudgeted* RootPrerequisite = FindRootPrerequisite(InComponent);
+		ComponentData.RootPrerequisite = (RootPrerequisite != nullptr && RootPrerequisite != InComponent) ? RootPrerequisite : nullptr;
 	}
 }
 
@@ -1107,6 +1099,8 @@ void FAnimationBudgetAllocator::HandleWorldBeginPlay()
 	// This will catch worlds in (e.g.) PIE that try to set the CVar on startup
 	FAnimationBudgetAllocator::bCachedEnabled = GAnimationBudgetEnabled == 1;
 
+	bHasBegunPlay = true;
+
 	// Run thru all deferred registrations
 	for(USkeletalMeshComponentBudgeted* Component : DeferredRegistrations)
 	{
diff --git a/Engine/Plugins/Runtime/AnimationBudgetAllocator/Source/AnimationBudgetAllocator/Private/AnimationBudgetAllocator.h b/Engine/Plugins/Runtime/AnimationBudgetAllocator/Source/AnimationBudgetAllocator/Private/AnimationBudgetAllocator.h
--- a/Engine/Plugins/Runtime/AnimationBudgetAllocator/Source/AnimationBudgetAllocator/Private/AnimationBudgetAllocator.h
+++ b/Engine/Plugins/Runtime/AnimationBudgetAllocator/Source/AnimationBudgetAllocator/Private/AnimationBudgetAllocator.h
@@ -155,7 +155,10 @@ public:
 
 	// Registers a component for deferred registration (post begin play)
 	void RegisterComponentDeferred(USkeletalMeshComponentBudgeted* InComponent);
-	
+
+	// Checks whether this allocator's world has begun play
+	bool HasBegunPlay() const { return bHasBegunPlay; } 
+
 protected:
 	/** We tick before all actors in the world using this delegate */
 	void OnWorldPreActorTick(UWorld* InWorld, ELevelTick InLevelTick, float InDeltaSeconds);
@@ -283,4 +286,7 @@ protected:
 
 	/** Cached enabled flag that is copied from the CVar each tick */
 	static bool bCachedEnabled;
+
+	/** Flag to indidcate whether begin play has been called for this allocator's world */
+	bool bHasBegunPlay;
 };
diff --git a/Engine/Plugins/Runtime/AnimationBudgetAllocator/Source/AnimationBudgetAllocator/Private/SkeletalMeshComponentBudgeted.cpp b/Engine/Plugins/Runtime/AnimationBudgetAllocator/Source/AnimationBudgetAllocator/Private/SkeletalMeshComponentBudgeted.cpp
--- a/Engine/Plugins/Runtime/AnimationBudgetAllocator/Source/AnimationBudgetAllocator/Private/SkeletalMeshComponentBudgeted.cpp
+++ b/Engine/Plugins/Runtime/AnimationBudgetAllocator/Source/AnimationBudgetAllocator/Private/SkeletalMeshComponentBudgeted.cpp
@@ -36,7 +36,7 @@ void USkeletalMeshComponentBudgeted::BeginPlay()
 		{
 			if (FAnimationBudgetAllocator* LocalAnimationBudgetAllocator = static_cast<FAnimationBudgetAllocator*>(IAnimationBudgetAllocator::Get(LocalWorld)))
 			{
-				if(LocalWorld->HasBegunPlay())
+				if(LocalAnimationBudgetAllocator->HasBegunPlay())
 				{
 					// World is playing, so register 
 					LocalAnimationBudgetAllocator->RegisterComponent(this);
