[Animation Budget Allocator] - Introduced a StartGarbageCollection removal that functions the exact same as the Post equivalent. This is because we could be accidentally holding onto Memory that has been Marked as Garbage elsewhere, which causes a warning.

#rnx

diff --git a/Engine/Plugins/Runtime/AnimationBudgetAllocator/Source/AnimationBudgetAllocator/Private/AnimationBudgetAllocator.cpp b/Engine/Plugins/Runtime/AnimationBudgetAllocator/Source/AnimationBudgetAllocator/Private/AnimationBudgetAllocator.cpp
--- a/Engine/Plugins/Runtime/AnimationBudgetAllocator/Source/AnimationBudgetAllocator/Private/AnimationBudgetAllocator.cpp
+++ b/Engine/Plugins/Runtime/AnimationBudgetAllocator/Source/AnimationBudgetAllocator/Private/AnimationBudgetAllocator.cpp
@@ -73,6 +73,7 @@ FAnimationBudgetAllocator::FAnimationBudgetAllocator(UWorld* InWorld)
 	SetParametersFromCVars();
 
 	OnWorldBeginPlayHandle = InWorld->OnWorldBeginPlay.AddRaw(this, &FAnimationBudgetAllocator::HandleWorldBeginPlay);
+	GarbageCollectStartedHandle = FCoreUObjectDelegates::GetGarbageCollectStartedDelegate().AddRaw(this, &FAnimationBudgetAllocator::HandleGarbageCollectStarted);
 	PostGarbageCollectHandle = FCoreUObjectDelegates::GetPostGarbageCollect().AddRaw(this, &FAnimationBudgetAllocator::HandlePostGarbageCollect);
 	OnWorldPreActorTickHandle = FWorldDelegates::OnWorldPreActorTick.AddRaw(this, &FAnimationBudgetAllocator::OnWorldPreActorTick);
 	OnCVarParametersChangedHandle = GOnCVarParametersChanged.AddRaw(this, &FAnimationBudgetAllocator::SetParametersFromCVars);
@@ -90,6 +91,8 @@ FAnimationBudgetAllocator::~FAnimationBudgetAllocator()
 	{
 		World->OnWorldBeginPlay.Remove(OnWorldBeginPlayHandle);
 	}
+
+	FCoreUObjectDelegates::GetGarbageCollectStartedDelegate().Remove(GarbageCollectStartedHandle);
 	FCoreUObjectDelegates::GetPostGarbageCollect().Remove(PostGarbageCollectHandle);
 	FWorldDelegates::OnWorldPreActorTick.Remove(OnWorldPreActorTickHandle);
 	GOnCVarParametersChanged.Remove(OnCVarParametersChangedHandle);
@@ -999,18 +1002,27 @@ void FAnimationBudgetAllocator::AddReferencedObjects(FReferenceCollector& Collec
 	}
 }
 
-void FAnimationBudgetAllocator::HandlePostGarbageCollect()
+void FAnimationBudgetAllocator::RemoveDeadComponents()
 {
-	// Remove dead components backwards, readjusting indices
 	for (int32 DataIndex = AllComponentData.Num() - 1; DataIndex >= 0; --DataIndex)
 	{
-		if (AllComponentData[DataIndex].Component == nullptr)
+		if (!IsValid(AllComponentData[DataIndex].Component))
 		{
 			RemoveHelper(DataIndex, nullptr);
 		}
 	}
 }
 
+void FAnimationBudgetAllocator::HandleGarbageCollectStarted()
+{
+	RemoveDeadComponents();
+}
+
+void FAnimationBudgetAllocator::HandlePostGarbageCollect()
+{
+	RemoveDeadComponents();
+}
+
 void FAnimationBudgetAllocator::SetGameThreadLastTickTimeMs(int32 InManagerHandle, float InGameThreadLastTickTimeMs)
 {
 	if(InManagerHandle != INDEX_NONE)
diff --git a/Engine/Plugins/Runtime/AnimationBudgetAllocator/Source/AnimationBudgetAllocator/Private/AnimationBudgetAllocator.h b/Engine/Plugins/Runtime/AnimationBudgetAllocator/Source/AnimationBudgetAllocator/Private/AnimationBudgetAllocator.h
--- a/Engine/Plugins/Runtime/AnimationBudgetAllocator/Source/AnimationBudgetAllocator/Private/AnimationBudgetAllocator.h
+++ b/Engine/Plugins/Runtime/AnimationBudgetAllocator/Source/AnimationBudgetAllocator/Private/AnimationBudgetAllocator.h
@@ -163,6 +163,12 @@ protected:
 	/** We tick before all actors in the world using this delegate */
 	void OnWorldPreActorTick(UWorld* InWorld, ELevelTick InLevelTick, float InDeltaSeconds);
 
+	/** Remove dead components backwards, readjusting indices */
+	void RemoveDeadComponents();
+
+	/** Clean up dead components at the start of GC, but before Lock. */
+	void HandleGarbageCollectStarted();
+
 	/** Clean up dead components post GC */
 	void HandlePostGarbageCollect();
 
@@ -263,7 +269,10 @@ protected:
 	/** Throttle counter for delaying reduced work */
 	int32 ReducedComponentWorkCounter;
 
-	/** Handle used to track garbage collection */
+	/** Handle used to track start of garbage collection */
+	FDelegateHandle GarbageCollectStartedHandle;
+
+	/** Handle used to track post-garbage collection */
 	FDelegateHandle PostGarbageCollectHandle;
 
 	/** Handle used for ticking */
