Added option to disable animation ticking when a skeletal mesh is directly occlusion/frustum culled

As opposed to fully culled, including shadow/Lumen rendering
This can result in mesh shadows freezing, but can be useful depending on the situation

PR #12193: Option to do not tick skel meshes while they're occlusion culled
[RN] Added option to disable animation ticking when a skeletal mesh is directly occlusion/frustum culled

diff --git a/Engine/Plugins/Runtime/AnimationBudgetAllocator/Source/AnimationBudgetAllocator/Private/AnimationBudgetAllocator.cpp b/Engine/Plugins/Runtime/AnimationBudgetAllocator/Source/AnimationBudgetAllocator/Private/AnimationBudgetAllocator.cpp
--- a/Engine/Plugins/Runtime/AnimationBudgetAllocator/Source/AnimationBudgetAllocator/Private/AnimationBudgetAllocator.cpp
+++ b/Engine/Plugins/Runtime/AnimationBudgetAllocator/Source/AnimationBudgetAllocator/Private/AnimationBudgetAllocator.cpp
@@ -238,7 +238,7 @@ void FAnimationBudgetAllocator::QueueSortedComponentIndices(float InDeltaSeconds
 			{
 				auto ShouldComponentTick = [WorldTime](const USkeletalMeshComponentBudgeted* InComponent, const FAnimBudgetAllocatorComponentData& InComponentData)
 				{
-					return ((InComponent->GetLastRenderTime() > WorldTime) ||
+					return (((InComponent->bUseScreenRenderStateForUpdate ? InComponent->GetLastRenderTimeOnScreen() : InComponent->GetLastRenderTime()) > WorldTime) ||
 						(InComponent->GetShouldUseActorRenderedFlag() && InComponent->GetAttachmentRootActor() && InComponent->GetAttachmentRootActor()->WasRecentlyRendered())  ||
 							InComponentData.bTickEvenIfNotRendered ||
 							InComponent->ShouldTickPose() ||
@@ -264,7 +264,7 @@ void FAnimationBudgetAllocator::QueueSortedComponentIndices(float InDeltaSeconds
 				{
 					// Push into a separate limited list if we are 'tick even if not rendered'.
 					// Skip offscreen components with a significance of zero or less.
-					if(Component->GetLastRenderTime() <= WorldTime && ComponentData.bTickEvenIfNotRendered && ComponentData.Significance > 0.0f)
+					if((Component->bUseScreenRenderStateForUpdate ? Component->GetLastRenderTimeOnScreen() : Component->GetLastRenderTime()) <= WorldTime && ComponentData.bTickEvenIfNotRendered && ComponentData.Significance > 0.0f)
 					{
 						NonRenderedComponentData.Add(ComponentIndex);
 					}
diff --git a/Engine/Source/Runtime/Engine/Classes/Components/SkinnedMeshComponent.h b/Engine/Source/Runtime/Engine/Classes/Components/SkinnedMeshComponent.h
--- a/Engine/Source/Runtime/Engine/Classes/Components/SkinnedMeshComponent.h
+++ b/Engine/Source/Runtime/Engine/Classes/Components/SkinnedMeshComponent.h
@@ -804,6 +804,15 @@ public:
 	UPROPERTY(transient)
 	uint8 bCPUSkinning : 1;
 
+	/** 
+	 * If set, use the screen render flag instead of the default render flag when processing offscreen-rendering optimizations 
+	 * (such as VisibilityBasedAnimTickOption) that look to reduce animation work when the mesh is not rendered. 
+	 * Using this option can result in meshes that are occlusion culled ceasing to perform animation work.
+	 * Note that this can however result in shadows not being animated when meshes are not directly visible.
+	 */
+	UPROPERTY(EditAnywhere, AdvancedDisplay, BlueprintReadWrite, Category=Optimization)
+	uint8 bUseScreenRenderStateForUpdate:1;
+
 	// Update Rate
 	/** if TRUE, Owner will determine how often animation will be updated and evaluated. See AnimUpdateRateTick() 
 	 * This allows to skip frames for performance. (For example based on visibility and size on screen). */
diff --git a/Engine/Source/Runtime/Engine/Private/Components/SkinnedMeshComponent.cpp b/Engine/Source/Runtime/Engine/Private/Components/SkinnedMeshComponent.cpp
--- a/Engine/Source/Runtime/Engine/Private/Components/SkinnedMeshComponent.cpp
+++ b/Engine/Source/Runtime/Engine/Private/Components/SkinnedMeshComponent.cpp
@@ -1471,7 +1471,7 @@ void USkinnedMeshComponent::TickComponent(float DeltaTime, enum ELevelTick TickT
 	}
 
 	// See if this mesh was rendered recently. This has to happen first because other data will rely on this
-	bRecentlyRendered = (GetLastRenderTime() > GetWorld()->TimeSeconds - 1.0f);
+	bRecentlyRendered = ((bUseScreenRenderStateForUpdate ? GetLastRenderTimeOnScreen() : GetLastRenderTime()) > GetWorld()->TimeSeconds - 1.0f);
 
 	// Update component's LOD settings
 	// This must be done BEFORE animation Update and Evaluate (TickPose and RefreshBoneTransforms respectively)
