Fixed animation budget allocator debug rendering sometimes getting removed

When ticking offscreen meshes (only possible with native code driving significance), sometimes an empty strting could be passed to debug rendering, removing all text for that actor. As we passed nullptr for the base actor, the world settings actor is used instead as the key, so passing an empty string cleared all debug text. The fix is to pass the correct base actor (and a zero vector for the location).

diff --git a/Engine/Plugins/Runtime/AnimationBudgetAllocator/Source/AnimationBudgetAllocator/Private/AnimationBudgetAllocator.cpp b/Engine/Plugins/Runtime/AnimationBudgetAllocator/Source/AnimationBudgetAllocator/Private/AnimationBudgetAllocator.cpp
--- a/Engine/Plugins/Runtime/AnimationBudgetAllocator/Source/AnimationBudgetAllocator/Private/AnimationBudgetAllocator.cpp
+++ b/Engine/Plugins/Runtime/AnimationBudgetAllocator/Source/AnimationBudgetAllocator/Private/AnimationBudgetAllocator.cpp
@@ -733,8 +733,6 @@ void FAnimationBudgetAllocator::Update(float DeltaSeconds)
 
 			for(const TPair<AActor*, TArray<int32>>& ActorIndicesPair : ActorMap)
 			{
-				FVector Location = ActorIndicesPair.Key->GetActorLocation();
-
 				FString DebugString;
 				
 				for(int32 ComponentDataIndex : ActorIndicesPair.Value)
@@ -753,7 +751,7 @@ void FAnimationBudgetAllocator::Update(float DeltaSeconds)
 					}
 				}
 
-				DrawDebugString(World, Location, DebugString, nullptr, FColor::White, 0.016f, false);
+				DrawDebugString(World, FVector::ZeroVector, DebugString, ActorIndicesPair.Key, FColor::White, 0.016f, false);
 			}
 
 			DebugTimes.Add(FVector2D(CurrentDebugTimeDisplay, DebugTotalTime));
