When sorting components in AnimationBudgetAllocator, prioritise components with neverskip.

diff --git a/Engine/Plugins/Runtime/AnimationBudgetAllocator/Source/AnimationBudgetAllocator/Private/AnimationBudgetAllocator.cpp b/Engine/Plugins/Runtime/AnimationBudgetAllocator/Source/AnimationBudgetAllocator/Private/AnimationBudgetAllocator.cpp
--- a/Engine/Plugins/Runtime/AnimationBudgetAllocator/Source/AnimationBudgetAllocator/Private/AnimationBudgetAllocator.cpp
+++ b/Engine/Plugins/Runtime/AnimationBudgetAllocator/Source/AnimationBudgetAllocator/Private/AnimationBudgetAllocator.cpp
@@ -300,6 +300,13 @@ void FAnimationBudgetAllocator::QueueSortedComponentIndices(float InDeltaSeconds
 	{
 		const FAnimBudgetAllocatorComponentData& ComponentData0 = AllComponentData[InIndex0];
 		const FAnimBudgetAllocatorComponentData& ComponentData1 = AllComponentData[InIndex1];
+
+		// prioritize component with always tick enabled
+		if (ComponentData0.bAlwaysTick != ComponentData1.bAlwaysTick)
+		{
+			return static_cast<bool>(ComponentData0.bAlwaysTick);
+		}
+
 		return ComponentData0.Significance > ComponentData1.Significance;
 	};
 
