Fix potential invalid reference within budgeted components to deleted allocator

diff --git a/Engine/Plugins/Runtime/AnimationBudgetAllocator/Source/AnimationBudgetAllocator/Private/AnimationBudgetAllocator.cpp b/Engine/Plugins/Runtime/AnimationBudgetAllocator/Source/AnimationBudgetAllocator/Private/AnimationBudgetAllocator.cpp
--- a/Engine/Plugins/Runtime/AnimationBudgetAllocator/Source/AnimationBudgetAllocator/Private/AnimationBudgetAllocator.cpp
+++ b/Engine/Plugins/Runtime/AnimationBudgetAllocator/Source/AnimationBudgetAllocator/Private/AnimationBudgetAllocator.cpp
@@ -983,6 +983,18 @@ void FAnimationBudgetAllocator::UnregisterComponent(USkeletalMeshComponentBudget
 	}
 }
 
+void FAnimationBudgetAllocator::UnregisterAllComponents()
+{
+	for (int32 DataIndex = AllComponentData.Num() - 1; DataIndex >= 0; --DataIndex)
+	{
+		FAnimBudgetAllocatorComponentData& ComponentData = AllComponentData[DataIndex];
+		if (IsValid(ComponentData.Component))
+		{
+			UnregisterComponent(ComponentData.Component);
+		}
+	}
+}
+
 void FAnimationBudgetAllocator::UpdateComponentTickPrerequsites(USkeletalMeshComponentBudgeted* InComponent)
 {
 	checkSlow(World == InComponent->GetWorld());
diff --git a/Engine/Plugins/Runtime/AnimationBudgetAllocator/Source/AnimationBudgetAllocator/Private/AnimationBudgetAllocator.h b/Engine/Plugins/Runtime/AnimationBudgetAllocator/Source/AnimationBudgetAllocator/Private/AnimationBudgetAllocator.h
--- a/Engine/Plugins/Runtime/AnimationBudgetAllocator/Source/AnimationBudgetAllocator/Private/AnimationBudgetAllocator.h
+++ b/Engine/Plugins/Runtime/AnimationBudgetAllocator/Source/AnimationBudgetAllocator/Private/AnimationBudgetAllocator.h
@@ -133,6 +133,7 @@ public:
 	// IAnimationBudgetAllocator interface
 	virtual void RegisterComponent(USkeletalMeshComponentBudgeted* InComponent) override;
 	virtual void UnregisterComponent(USkeletalMeshComponentBudgeted* InComponent) override;
+	virtual void UnregisterAllComponents() override;
 	virtual void UpdateComponentTickPrerequsites(USkeletalMeshComponentBudgeted* InComponent) override;
 	virtual void SetComponentSignificance(USkeletalMeshComponentBudgeted* Component, float Significance, bool bNeverSkip = false, bool bTickEvenIfNotRendered = false, bool bAllowReducedWork = true, bool bForceInterpolate = false) override;
 	virtual void SetComponentTickEnabled(USkeletalMeshComponentBudgeted* Component, bool bShouldTick) override;
diff --git a/Engine/Plugins/Runtime/AnimationBudgetAllocator/Source/AnimationBudgetAllocator/Private/AnimationBudgetAllocatorModule.cpp b/Engine/Plugins/Runtime/AnimationBudgetAllocator/Source/AnimationBudgetAllocator/Private/AnimationBudgetAllocatorModule.cpp
--- a/Engine/Plugins/Runtime/AnimationBudgetAllocator/Source/AnimationBudgetAllocator/Private/AnimationBudgetAllocatorModule.cpp
+++ b/Engine/Plugins/Runtime/AnimationBudgetAllocator/Source/AnimationBudgetAllocator/Private/AnimationBudgetAllocatorModule.cpp
@@ -52,6 +52,7 @@ void FAnimationBudgetAllocatorModule::HandleWorldCleanup(UWorld* World, bool bSe
 		//Cleanup components, as this can cause a crash in some conditions during seamless travel as actors can be pulled in
 		//to the travel world.
 		Budgeter->SetEnabled(false);
+		Budgeter->UnregisterAllComponents();
 		delete Budgeter;
 		WorldAnimationBudgetAllocators.Remove(World);
 	}
diff --git a/Engine/Plugins/Runtime/AnimationBudgetAllocator/Source/AnimationBudgetAllocator/Public/IAnimationBudgetAllocator.h b/Engine/Plugins/Runtime/AnimationBudgetAllocator/Source/AnimationBudgetAllocator/Public/IAnimationBudgetAllocator.h
--- a/Engine/Plugins/Runtime/AnimationBudgetAllocator/Source/AnimationBudgetAllocator/Public/IAnimationBudgetAllocator.h
+++ b/Engine/Plugins/Runtime/AnimationBudgetAllocator/Source/AnimationBudgetAllocator/Public/IAnimationBudgetAllocator.h
@@ -35,6 +35,11 @@ public:
 	 */
 	virtual void UnregisterComponent(USkeletalMeshComponentBudgeted* InComponent) = 0;
 
+	/**
+	 * Calls UnregisterComponent on all components currently registered with the allocator.
+	 */
+	virtual void UnregisterAllComponents() = 0;
+
 	/**
 	 * Update the prerequisites of this component. Should be called when prerequisites may have changed externally.
 	 */
