Wrap deprecated property WITH_EDITORONLY_DATA to reduce waste
Bump SkeletalMeshComponent warning size

diff --git a/Engine/Plugins/Runtime/AnimationBudgetAllocator/Source/AnimationBudgetAllocator/Private/SkeletalMeshComponentBudgeted.cpp b/Engine/Plugins/Runtime/AnimationBudgetAllocator/Source/AnimationBudgetAllocator/Private/SkeletalMeshComponentBudgeted.cpp
--- a/Engine/Plugins/Runtime/AnimationBudgetAllocator/Source/AnimationBudgetAllocator/Private/SkeletalMeshComponentBudgeted.cpp
+++ b/Engine/Plugins/Runtime/AnimationBudgetAllocator/Source/AnimationBudgetAllocator/Private/SkeletalMeshComponentBudgeted.cpp
@@ -13,7 +13,7 @@
 CSV_DECLARE_CATEGORY_EXTERN(AnimationBudget);
 
 #if UE_BUILD_SHIPPING && !WITH_EDITOR
-static_assert(sizeof(USkeletalMeshComponentBudgeted) <= 4080, "We expect SkeletalMeshComponentBudgeted to be 4080 or less, otherwise we waste memory on 288bytes of padding on MallocBinned2. Please improve padding on this class hierarchy to pay for your added variable.");
+static_assert(sizeof(USkeletalMeshComponentBudgeted) <= 4368, "We expect SkeletalMeshComponentBudgeted to be 4368 or less, otherwise we waste memory on 304bytes of padding on MallocBinned2. Please improve padding on this class hierarchy to pay for your added variable.");
 #endif
 
 FOnCalculateSignificance USkeletalMeshComponentBudgeted::OnCalculateSignificanceDelegate;
