Prevent auto calculate significance overriding budget allocator component data flags

diff --git a/Engine/Plugins/Runtime/AnimationBudgetAllocator/Source/AnimationBudgetAllocator/Private/AnimationBudgetAllocator.cpp b/Engine/Plugins/Runtime/AnimationBudgetAllocator/Source/AnimationBudgetAllocator/Private/AnimationBudgetAllocator.cpp
--- a/Engine/Plugins/Runtime/AnimationBudgetAllocator/Source/AnimationBudgetAllocator/Private/AnimationBudgetAllocator.cpp
+++ b/Engine/Plugins/Runtime/AnimationBudgetAllocator/Source/AnimationBudgetAllocator/Private/AnimationBudgetAllocator.cpp
@@ -212,7 +212,12 @@ void FAnimationBudgetAllocator::QueueSortedComponentIndices(float InDeltaSeconds
 				Significance = InComponentData.Component->GetDefaultSignificance();
 			}
 			
-			SetComponentSignificance(InComponentData.Component, Significance);
+			SetComponentSignificance(InComponentData.Component,
+				Significance,
+				InComponentData.bAlwaysTick,
+				InComponentData.bTickEvenIfNotRendered,
+				InComponentData.bAllowReducedWork,
+				InComponentData.bNeverThrottle);
 		}
 	};
 
