PR: Fix incorrect comparison of VisibilityBasedAnimTickOption in AnimationBudgetAllocator

diff --git a/Engine/Plugins/Runtime/AnimationBudgetAllocator/Source/AnimationBudgetAllocator/Private/AnimationBudgetAllocator.cpp b/Engine/Plugins/Runtime/AnimationBudgetAllocator/Source/AnimationBudgetAllocator/Private/AnimationBudgetAllocator.cpp
--- a/Engine/Plugins/Runtime/AnimationBudgetAllocator/Source/AnimationBudgetAllocator/Private/AnimationBudgetAllocator.cpp
+++ b/Engine/Plugins/Runtime/AnimationBudgetAllocator/Source/AnimationBudgetAllocator/Private/AnimationBudgetAllocator.cpp
@@ -248,7 +248,8 @@ void FAnimationBudgetAllocator::QueueSortedComponentIndices(float InDeltaSeconds
 							InComponentData.bTickEvenIfNotRendered ||
 							InComponent->ShouldTickPose() ||
 							InComponent->ShouldUpdateTransform(false) ||	// We can force this to false, only used in WITH_EDITOR
-							InComponent->VisibilityBasedAnimTickOption == EVisibilityBasedAnimTickOption::AlwaysTickPose);
+							InComponent->VisibilityBasedAnimTickOption == EVisibilityBasedAnimTickOption::AlwaysTickPose ||
+							InComponent->VisibilityBasedAnimTickOption == EVisibilityBasedAnimTickOption::AlwaysTickPoseAndRefreshBones);
 				};
 
 				// Whether or not we will tick
