Animationbudgetallocator now resets external tick state on skeletal mesh components when disabled via cvar.

diff --git a/Engine/Plugins/Runtime/AnimationBudgetAllocator/Source/AnimationBudgetAllocator/Private/AnimationBudgetAllocator.cpp b/Engine/Plugins/Runtime/AnimationBudgetAllocator/Source/AnimationBudgetAllocator/Private/AnimationBudgetAllocator.cpp
--- a/Engine/Plugins/Runtime/AnimationBudgetAllocator/Source/AnimationBudgetAllocator/Private/AnimationBudgetAllocator.cpp
+++ b/Engine/Plugins/Runtime/AnimationBudgetAllocator/Source/AnimationBudgetAllocator/Private/AnimationBudgetAllocator.cpp
@@ -66,6 +66,7 @@ FAnimationBudgetAllocator::FAnimationBudgetAllocator(UWorld* InWorld)
 	, ReducedComponentWorkCounter(0)
 	, CurrentFrameOffset(0)
 	, bEnabled(false)
+	, bCurrentlyUpdating(false)
 	, bHasBegunPlay(InWorld->HasBegunPlay())
 {
 	FAnimationBudgetAllocator::bCachedEnabled = GAnimationBudgetEnabled == 1;
@@ -706,9 +707,9 @@ void FAnimationBudgetAllocator::Update(float DeltaSeconds)
 	SCOPE_CYCLE_COUNTER(STAT_AnimationBudgetAllocator_Update);
 	CSV_SCOPED_TIMING_STAT_EXCLUSIVE(AnimationBudgetAllocator);
 
-	FAnimationBudgetAllocator::bCachedEnabled = GAnimationBudgetEnabled == 1;
+	CheckAndPropagateUpdateStatusChange();
 
-	if (FAnimationBudgetAllocator::bCachedEnabled && bEnabled)
+	if (bCurrentlyUpdating)
 	{
 		check(IsInGameThread());
 
@@ -1085,34 +1086,7 @@ void FAnimationBudgetAllocator::SetIsRunningReducedWork(USkeletalMeshComponentBu
 void FAnimationBudgetAllocator::SetEnabled(bool bInEnabled)
 {
 	bEnabled = bInEnabled;
-
-	if (bEnabled)
-	{
-		for(int32 DataIndex = 0; DataIndex < AllComponentData.Num(); ++DataIndex)
-		{
-			FAnimBudgetAllocatorComponentData& ComponentData = AllComponentData[DataIndex];
-			if(ComponentData.Component != nullptr)
-			{
-				ComponentData.Component->bEnableUpdateRateOptimizations = false;
-				ComponentData.Component->EnableExternalTickRateControl(true);
-			}
-		}
-	}
-	else
-	{
-		for(int32 DataIndex = 0; DataIndex < AllComponentData.Num(); ++DataIndex)
-		{
-			FAnimBudgetAllocatorComponentData& ComponentData = AllComponentData[DataIndex];
-			if(ComponentData.Component != nullptr)
-			{
-				ComponentData.Component->bEnableUpdateRateOptimizations = true;
-				ComponentData.Component->EnableExternalTickRateControl(false);
-
-				// Re-enable ticking in case the component was currently skipping
-				TickEnableHelper(ComponentData.Component, true);
-			}
-		}
-	}
+	CheckAndPropagateUpdateStatusChange();
 }
 
 bool FAnimationBudgetAllocator::GetEnabled() const
@@ -1139,6 +1113,45 @@ void FAnimationBudgetAllocator::ForceNextTickThisFrame(USkeletalMeshComponentBud
 	}
 }
 
+void FAnimationBudgetAllocator::CheckAndPropagateUpdateStatusChange()
+{
+	FAnimationBudgetAllocator::bCachedEnabled = GAnimationBudgetEnabled == 1;
+	const bool bUpdatingThisTick = FAnimationBudgetAllocator::bCachedEnabled && bEnabled;
+
+	if (bUpdatingThisTick != bCurrentlyUpdating)
+	{
+		if (bUpdatingThisTick)
+		{
+			for (int32 DataIndex = 0; DataIndex < AllComponentData.Num(); ++DataIndex)
+			{
+				FAnimBudgetAllocatorComponentData& ComponentData = AllComponentData[DataIndex];
+				if (ComponentData.Component != nullptr)
+				{
+					ComponentData.Component->bEnableUpdateRateOptimizations = false;
+					ComponentData.Component->EnableExternalTickRateControl(true);
+				}
+			}
+		}
+		else
+		{
+			for (int32 DataIndex = 0; DataIndex < AllComponentData.Num(); ++DataIndex)
+			{
+				FAnimBudgetAllocatorComponentData& ComponentData = AllComponentData[DataIndex];
+				if (ComponentData.Component != nullptr)
+				{
+					ComponentData.Component->bEnableUpdateRateOptimizations = true;
+					ComponentData.Component->EnableExternalTickRateControl(false);
+
+					// Re-enable ticking in case the component was currently skipping
+					TickEnableHelper(ComponentData.Component, true);
+				}
+			}
+		}
+	}
+
+	bCurrentlyUpdating = bUpdatingThisTick;
+}
+
 void FAnimationBudgetAllocator::SetParametersFromCVars()
 {
 	Parameters = GBudgetParameters;
diff --git a/Engine/Plugins/Runtime/AnimationBudgetAllocator/Source/AnimationBudgetAllocator/Private/AnimationBudgetAllocator.h b/Engine/Plugins/Runtime/AnimationBudgetAllocator/Source/AnimationBudgetAllocator/Private/AnimationBudgetAllocator.h
--- a/Engine/Plugins/Runtime/AnimationBudgetAllocator/Source/AnimationBudgetAllocator/Private/AnimationBudgetAllocator.h
+++ b/Engine/Plugins/Runtime/AnimationBudgetAllocator/Source/AnimationBudgetAllocator/Private/AnimationBudgetAllocator.h
@@ -194,6 +194,9 @@ protected:
 	/** Initializes internal parameters from CVars. */
 	void SetParametersFromCVars();
 
+	/** check for change in update status and set component external tick states to match */
+	void CheckAndPropagateUpdateStatusChange();
+
 #if ENABLE_DRAW_DEBUG
 	/** Debug support function */
 	void OnHUDPostRender(AHUD* HUD, UCanvas* Canvas);
@@ -297,6 +300,9 @@ protected:
 	/** Cached enabled flag that is copied from the CVar each tick */
 	static bool bCachedEnabled;
 
+	/** Flag to indicate allocator is currently updating (ie cvar enabled and local enabled both true) */
+	bool bCurrentlyUpdating;
+
 	/** Flag to indidcate whether begin play has been called for this allocator's world */
 	bool bHasBegunPlay;
 };
