Fix for memory leak in AnimationBudgetAllocator

diff --git a/Engine/Plugins/Runtime/AnimationBudgetAllocator/Source/AnimationBudgetAllocator/Private/AnimationBudgetAllocator.cpp b/Engine/Plugins/Runtime/AnimationBudgetAllocator/Source/AnimationBudgetAllocator/Private/AnimationBudgetAllocator.cpp
--- a/Engine/Plugins/Runtime/AnimationBudgetAllocator/Source/AnimationBudgetAllocator/Private/AnimationBudgetAllocator.cpp
+++ b/Engine/Plugins/Runtime/AnimationBudgetAllocator/Source/AnimationBudgetAllocator/Private/AnimationBudgetAllocator.cpp
@@ -91,6 +91,7 @@ FAnimationBudgetAllocator::~FAnimationBudgetAllocator()
 	if(World)
 	{
 		World->OnWorldBeginPlay.Remove(OnWorldBeginPlayHandle);
+		World = nullptr;
 	}
 
 	FCoreUObjectDelegates::GetGarbageCollectStartedDelegate().Remove(GarbageCollectStartedHandle);
