Animation Budget Allocator: Enable component tick when unregistering

PR #11733: Enable component tick while unregister from the AnimationBudgetAllocator

diff --git a/Engine/Plugins/Runtime/AnimationBudgetAllocator/Source/AnimationBudgetAllocator/Private/AnimationBudgetAllocator.cpp b/Engine/Plugins/Runtime/AnimationBudgetAllocator/Source/AnimationBudgetAllocator/Private/AnimationBudgetAllocator.cpp
--- a/Engine/Plugins/Runtime/AnimationBudgetAllocator/Source/AnimationBudgetAllocator/Private/AnimationBudgetAllocator.cpp
+++ b/Engine/Plugins/Runtime/AnimationBudgetAllocator/Source/AnimationBudgetAllocator/Private/AnimationBudgetAllocator.cpp
@@ -988,6 +988,8 @@ void FAnimationBudgetAllocator::UnregisterComponent(USkeletalMeshComponentBudget
 		InComponent->bEnableUpdateRateOptimizations = true;
 		InComponent->EnableExternalTickRateControl(false);
 		InComponent->SetAnimationBudgetAllocator(nullptr);
+
+		TickEnableHelper(InComponent, true);
 	}
 }
 
