[Animation]  Update Fixes and Changes
- Add Option to ABA to allow external anim source time

#changelist validated

diff --git a/Engine/Plugins/Runtime/AnimationBudgetAllocator/Source/AnimationBudgetAllocator/Private/AnimationBudgetAllocator.cpp b/Engine/Plugins/Runtime/AnimationBudgetAllocator/Source/AnimationBudgetAllocator/Private/AnimationBudgetAllocator.cpp
--- a/Engine/Plugins/Runtime/AnimationBudgetAllocator/Source/AnimationBudgetAllocator/Private/AnimationBudgetAllocator.cpp
+++ b/Engine/Plugins/Runtime/AnimationBudgetAllocator/Source/AnimationBudgetAllocator/Private/AnimationBudgetAllocator.cpp
@@ -481,7 +481,11 @@ int32 FAnimationBudgetAllocator::CalculateWorkDistributionAndQueue(float InDelta
 
 			// Want to map the remaining (non-fixed) work units so that we only execute N work units per frame.
 			// If we can go over budget to keep quality then we use that value
-			const float WorkUnitBudget = FMath::Max(Parameters.BudgetInMs / AverageWorkUnitTimeMs, (float)TotalIdealWorkUnits * Parameters.MinQuality);
+
+			const float AvailableBudget = Parameters.BudgetInMs - ExternalAnimationTimeSpent;
+			ExternalAnimationTimeSpent = 0.0f;
+
+			const float WorkUnitBudget = FMath::Max(AvailableBudget / AverageWorkUnitTimeMs, (float)TotalIdealWorkUnits * Parameters.MinQuality);
 
 			SET_FLOAT_STAT(STAT_AnimationBudgetAllocator_Budget, WorkUnitBudget);
 
@@ -1116,6 +1120,11 @@ void FAnimationBudgetAllocator::ForceNextTickThisFrame(USkeletalMeshComponentBud
 	}
 }
 
+void FAnimationBudgetAllocator::AddExternalAnimationTimeSpentThisFrame(float TimeSpent)
+{
+	ExternalAnimationTimeSpent += TimeSpent;
+}
+
 void FAnimationBudgetAllocator::CheckAndPropagateUpdateStatusChange()
 {
 	FAnimationBudgetAllocator::bCachedEnabled = GAnimationBudgetEnabled == 1;
diff --git a/Engine/Plugins/Runtime/AnimationBudgetAllocator/Source/AnimationBudgetAllocator/Private/AnimationBudgetAllocator.h b/Engine/Plugins/Runtime/AnimationBudgetAllocator/Source/AnimationBudgetAllocator/Private/AnimationBudgetAllocator.h
--- a/Engine/Plugins/Runtime/AnimationBudgetAllocator/Source/AnimationBudgetAllocator/Private/AnimationBudgetAllocator.h
+++ b/Engine/Plugins/Runtime/AnimationBudgetAllocator/Source/AnimationBudgetAllocator/Private/AnimationBudgetAllocator.h
@@ -148,6 +148,7 @@ public:
 	UE_API virtual bool GetEnabled() const override;
 	UE_API virtual void SetParameters(const FAnimationBudgetAllocatorParameters& InParameters) override;
 	UE_API virtual void ForceNextTickThisFrame(USkeletalMeshComponentBudgeted* Component) override;
+	UE_API virtual void AddExternalAnimationTimeSpentThisFrame(float TimeSpent) override;
 
 	// FGCObject interface
 	UE_API virtual void AddReferencedObjects(FReferenceCollector& Collector) override;
@@ -307,6 +308,10 @@ protected:
 
 	/** Flag to indidcate whether begin play has been called for this allocator's world */
 	bool bHasBegunPlay;
+
+	/** External animation time spent by sources not captured in ABA but that should reduce avaialbe ABA budget for the frame */
+	float ExternalAnimationTimeSpent = 0.0f;
+
 };
 
 #undef UE_API
diff --git a/Engine/Plugins/Runtime/AnimationBudgetAllocator/Source/AnimationBudgetAllocator/Public/IAnimationBudgetAllocator.h b/Engine/Plugins/Runtime/AnimationBudgetAllocator/Source/AnimationBudgetAllocator/Public/IAnimationBudgetAllocator.h
--- a/Engine/Plugins/Runtime/AnimationBudgetAllocator/Source/AnimationBudgetAllocator/Public/IAnimationBudgetAllocator.h
+++ b/Engine/Plugins/Runtime/AnimationBudgetAllocator/Source/AnimationBudgetAllocator/Public/IAnimationBudgetAllocator.h
@@ -80,4 +80,10 @@ public:
 
 	/** Force component to tick this frame regardless of tick rate */
 	virtual void ForceNextTickThisFrame(USkeletalMeshComponentBudgeted* Component) = 0;
+
+	/**
+	* Add external time spent on animation that should come out of the Animation Budget but is not accounted for in the ABA update.
+	* TimeSpent is accumulated and spent each frame, it is reset to 0 each time it is used.  
+	*/
+	virtual void AddExternalAnimationTimeSpentThisFrame(float TimeSpent) = 0;
 };
\ No newline at end of file
