Optional changes to AnimationBudgetAllocator (disabled by default):
- Maintain reducework on non-rendered offscreen components.
- Components with significance zero can be treated as non-rendered.

#Synced-CL 50155597

#Synced-CL 50155597

diff --git a/Engine/Plugins/Runtime/AnimationBudgetAllocator/Source/AnimationBudgetAllocator/Private/AnimationBudgetAllocator.cpp b/Engine/Plugins/Runtime/AnimationBudgetAllocator/Source/AnimationBudgetAllocator/Private/AnimationBudgetAllocator.cpp
--- a/Engine/Plugins/Runtime/AnimationBudgetAllocator/Source/AnimationBudgetAllocator/Private/AnimationBudgetAllocator.cpp
+++ b/Engine/Plugins/Runtime/AnimationBudgetAllocator/Source/AnimationBudgetAllocator/Private/AnimationBudgetAllocator.cpp
@@ -238,6 +238,8 @@ void FAnimationBudgetAllocator::QueueSortedComponentIndices(float InDeltaSeconds
 	int32 ComponentIndex = 0;
 	for (FAnimBudgetAllocatorComponentData& ComponentData : AllComponentData)
 	{
+		ComponentData.bReducedWorkDueToOffscreen = false;
+
 		if(ComponentData.bTickEnabled)
 		{
 			USkeletalMeshComponentBudgeted* Component = ComponentData.Component;
@@ -271,8 +273,10 @@ void FAnimationBudgetAllocator::QueueSortedComponentIndices(float InDeltaSeconds
 				if(bShouldTick)
 				{
 					// Push into a separate limited list if we are 'tick even if not rendered'.
-					// Skip offscreen components with a significance of zero or less.
-					if((Component->bUseScreenRenderStateForUpdate ? Component->GetLastRenderTimeOnScreen() : Component->GetLastRenderTime()) <= WorldTime && ComponentData.bTickEvenIfNotRendered && ComponentData.Significance > 0.0f)
+					// Skip offscreen components with a significance of zero or less (or less than zero if GAnimationBudgetComponentsWithZeroSigAreNonRendered is enabled).
+					if((Component->bUseScreenRenderStateForUpdate ? Component->GetLastRenderTimeOnScreen() : Component->GetLastRenderTime()) <= WorldTime && 
+						ComponentData.bTickEvenIfNotRendered && 
+						((GAnimationBudgetComponentsWithZeroSigAreNonRendered && ComponentData.Significance >= 0.0f) || (!GAnimationBudgetComponentsWithZeroSigAreNonRendered && ComponentData.Significance > 0.0f)))
 					{
 						NonRenderedComponentData.Add(ComponentIndex);
 					}
@@ -327,13 +331,18 @@ void FAnimationBudgetAllocator::QueueSortedComponentIndices(float InDeltaSeconds
 	{
 		auto ReduceWorkForOffscreenComponent = [](FAnimBudgetAllocatorComponentData& InComponentData)
 		{
-			if(InComponentData.bAllowReducedWork && !InComponentData.bReducedWork && InComponentData.Component->OnReduceWork().IsBound())
+			if(InComponentData.bAllowReducedWork) 
 			{
+				InComponentData.bReducedWorkDueToOffscreen = true;
+
+				if (!InComponentData.bReducedWork && InComponentData.Component->OnReduceWork().IsBound())
+				{
 #if WITH_TICK_DEBUG
-				UE_LOG(LogTemp, Warning, TEXT("Force-decreasing offscreen component work (mesh %s) (actor %llx)"), InComponentData.Component->SkeletalMesh ? *InComponentData.Component->SkeletalMesh->GetName() : TEXT("null"), (uint64)InComponentData.Component->GetOwner());
+					UE_LOG(LogTemp, Warning, TEXT("Force-decreasing offscreen component work (mesh %s) (actor %llx)"), InComponentData.Component->SkeletalMesh ? *InComponentData.Component->SkeletalMesh->GetName() : TEXT("null"), (uint64)InComponentData.Component->GetOwner());
 #endif
-				InComponentData.Component->OnReduceWork().Execute(InComponentData.Component, true);
-				InComponentData.bReducedWork = true;
+					InComponentData.Component->OnReduceWork().Execute(InComponentData.Component, true);
+					InComponentData.bReducedWork = true;
+				}
 			}
 		};
 
@@ -364,6 +373,20 @@ void FAnimationBudgetAllocator::QueueSortedComponentIndices(float InDeltaSeconds
 			ComponentData.StateChangeThrottle = -1;
 		}
 
+		if (GAnimationBudgetMaintainReduceWorkWhenOffScreen)
+		{
+			// if component ticks were disabled either here of through reduce work, they should be taken out of the sorted component list
+			for (int32 SortedCompIdx = 0 ; SortedCompIdx < AllSortedComponentData.Num() ; ++SortedCompIdx)
+			{
+				FAnimBudgetAllocatorComponentData& SortedComponentData = AllComponentData[AllSortedComponentData[SortedCompIdx]];
+				if (!SortedComponentData.bTickEnabled)
+				{
+					AllSortedComponentData.RemoveAt(SortedCompIdx);
+					--SortedCompIdx;
+				}
+			}
+		}
+
 		// re-sort now we have inserted offscreen components
 		AllSortedComponentData.Sort(SignificanceSortPredicate);
 	}
@@ -635,7 +658,25 @@ int32 FAnimationBudgetAllocator::CalculateWorkDistributionAndQueue(float InDelta
 					for(int32 ReducedWorkComponentIndex : ReducedWorkComponentData)
 					{
 						FAnimBudgetAllocatorComponentData& ComponentData = AllComponentData[ReducedWorkComponentIndex];
-						if(ComponentData.bReducedWork && ComponentData.Component->OnReduceWork().IsBound())
+
+						// check if ReducedWorkDueToOffscreen is true from this component or the root component. If so, don't remove reduce work due to budget
+						bool bReducedWorkDueToOffscreen = false;
+						
+						if (GAnimationBudgetMaintainReduceWorkWhenOffScreen)
+						{
+							bReducedWorkDueToOffscreen = ComponentData.bReducedWorkDueToOffscreen;
+							if (ComponentData.RootPrerequisite != nullptr)
+							{
+								const int32 PrerequisiteHandle = ComponentData.RootPrerequisite->GetAnimationBudgetHandle();
+								if (PrerequisiteHandle != INDEX_NONE)
+								{
+									FAnimBudgetAllocatorComponentData& RootPrerequisiteComponentData = AllComponentData[PrerequisiteHandle];
+									bReducedWorkDueToOffscreen |= RootPrerequisiteComponentData.bReducedWorkDueToOffscreen;
+								}
+							}
+						}
+
+						if(ComponentData.bReducedWork && !bReducedWorkDueToOffscreen && ComponentData.Component->OnReduceWork().IsBound())
 						{
 #if WITH_TICK_DEBUG
 							UE_LOG(LogTemp, Warning, TEXT("Increasing component work (mesh %s) (actor %llx)"), ComponentData.Component->SkeletalMesh ? *ComponentData.Component->SkeletalMesh->GetName() : TEXT("null"), (uint64)ComponentData.Component->GetOwner());
diff --git a/Engine/Plugins/Runtime/AnimationBudgetAllocator/Source/AnimationBudgetAllocator/Private/AnimationBudgetAllocator.h b/Engine/Plugins/Runtime/AnimationBudgetAllocator/Source/AnimationBudgetAllocator/Private/AnimationBudgetAllocator.h
--- a/Engine/Plugins/Runtime/AnimationBudgetAllocator/Source/AnimationBudgetAllocator/Private/AnimationBudgetAllocator.h
+++ b/Engine/Plugins/Runtime/AnimationBudgetAllocator/Source/AnimationBudgetAllocator/Private/AnimationBudgetAllocator.h
@@ -51,6 +51,7 @@ struct FAnimBudgetAllocatorComponentData
 		, bTickEvenIfNotRendered(false)
 		, bInterpolate(false)
 		, bReducedWork(false)
+		, bReducedWorkDueToOffscreen(false)
 		, bAllowReducedWork(false)
 		, bAutoCalculateSignificance(false)
 		, bOnScreen(false)
@@ -113,6 +114,9 @@ public:
 	/** Whether this component is running 'reduced work' */
 	uint8 bReducedWork : 1;
 
+	/** Whether this component is running 'reduced work' due to being non-rendered offscreen */
+	uint8 bReducedWorkDueToOffscreen : 1;
+
 	/** Whether this component allows 'reduced work' */
 	uint8 bAllowReducedWork : 1;
 
diff --git a/Engine/Plugins/Runtime/AnimationBudgetAllocator/Source/AnimationBudgetAllocator/Private/AnimationBudgetAllocatorCVars.cpp b/Engine/Plugins/Runtime/AnimationBudgetAllocator/Source/AnimationBudgetAllocator/Private/AnimationBudgetAllocatorCVars.cpp
--- a/Engine/Plugins/Runtime/AnimationBudgetAllocator/Source/AnimationBudgetAllocator/Private/AnimationBudgetAllocatorCVars.cpp
+++ b/Engine/Plugins/Runtime/AnimationBudgetAllocator/Source/AnimationBudgetAllocator/Private/AnimationBudgetAllocatorCVars.cpp
@@ -13,6 +13,26 @@ static FAutoConsoleVariableRef CVarSkelBatch_Enabled(
 	TEXT("Controls whether the skeletal mesh batching system is enabled. Should be set when there are no running skeletal meshes."),
 	ECVF_Scalability);
 
+// some optional changes to behavior in the budgeter
+
+int32 GAnimationBudgetMaintainReduceWorkWhenOffScreen = 0;
+
+static FAutoConsoleVariableRef CVarSkelBatch_MaintainReduceWorkWhenOffScreen(
+	TEXT("a.Budget.MaintainReduceWorkWhenOffScreen"),
+	GAnimationBudgetMaintainReduceWorkWhenOffScreen,
+	TEXT("Values: 0/1\n")
+	TEXT("Whether to enable behavior in the animation budgeter where non-rendered components keep reduce work rather than lose it when budget is available."),
+	ECVF_Scalability);
+
+int32 GAnimationBudgetComponentsWithZeroSigAreNonRendered = 0;
+
+static FAutoConsoleVariableRef CVarSkelBatch_ComponentsWithZeroSigAreNonRendered(
+	TEXT("a.Budget.ComponentsWithZeroSigAreNonRendered"),
+	GAnimationBudgetComponentsWithZeroSigAreNonRendered,
+	TEXT("Values: 0/1\n")
+	TEXT("Control whether components with zero significance can be treated in the animation budgeter as non-rendered."),
+	ECVF_Scalability);
+
 #if ENABLE_DRAW_DEBUG
 int32 GAnimationBudgetDebugEnabled = 0;
 
diff --git a/Engine/Plugins/Runtime/AnimationBudgetAllocator/Source/AnimationBudgetAllocator/Private/AnimationBudgetAllocatorCVars.h b/Engine/Plugins/Runtime/AnimationBudgetAllocator/Source/AnimationBudgetAllocator/Private/AnimationBudgetAllocatorCVars.h
--- a/Engine/Plugins/Runtime/AnimationBudgetAllocator/Source/AnimationBudgetAllocator/Private/AnimationBudgetAllocatorCVars.h
+++ b/Engine/Plugins/Runtime/AnimationBudgetAllocator/Source/AnimationBudgetAllocator/Private/AnimationBudgetAllocatorCVars.h
@@ -9,6 +9,10 @@ struct FAnimationBudgetAllocatorParameters;
 /** Enabled/disabled flag */
 extern int32 GAnimationBudgetEnabled;
 
+extern int32 GAnimationBudgetMaintainReduceWorkWhenOffScreen;
+
+extern int32 GAnimationBudgetComponentsWithZeroSigAreNonRendered;
+
 #if ENABLE_DRAW_DEBUG
 /** Debug rendering flag */
 extern int32 GAnimationBudgetDebugEnabled;
