Changes to the anim budgeter:
- Added option for a component that exits offscreen to immediately tick regardless of its current tick rate. This avoids us seeing a few frames of an old pose when a component re-enters view.
- Added debug option to exclude always tick components from forced budget.
- Fix for reducework constantly being turned on for offscreen components when using forced budget option.
- Added a whole bunch of verbose debug output options to show e.g. component names, non ticking components, etc.

#Synced-CL 50813309

diff --git a/Engine/Plugins/Runtime/AnimationBudgetAllocator/Source/AnimationBudgetAllocator/Private/AnimationBudgetAllocator.cpp b/Engine/Plugins/Runtime/AnimationBudgetAllocator/Source/AnimationBudgetAllocator/Private/AnimationBudgetAllocator.cpp
--- a/Engine/Plugins/Runtime/AnimationBudgetAllocator/Source/AnimationBudgetAllocator/Private/AnimationBudgetAllocator.cpp
+++ b/Engine/Plugins/Runtime/AnimationBudgetAllocator/Source/AnimationBudgetAllocator/Private/AnimationBudgetAllocator.cpp
@@ -23,7 +23,6 @@ DECLARE_DWORD_COUNTER_STAT(TEXT("Throttled"), STAT_AnimationBudgetAllocator_Thro
 DECLARE_DWORD_COUNTER_STAT(TEXT("Interpolated"), STAT_AnimationBudgetAllocator_Interpolated, STATGROUP_AnimationBudgetAllocator);
 DECLARE_FLOAT_COUNTER_STAT(TEXT("SmoothedBudgetPressure"), STAT_AnimationBudgetAllocator_SmoothedBudgetPressure, STATGROUP_AnimationBudgetAllocator);
 
-
 CSV_DEFINE_CATEGORY(AnimationBudget, true);
 
 bool FAnimationBudgetAllocator::bCachedEnabled = false;
@@ -35,6 +34,9 @@ FAnimBudgetAllocatorComponentData::FAnimBudgetAllocatorComponentData(USkeletalMe
 	, AccumulatedDeltaTime(0.0f)
 	, GameThreadLastTickTimeMs(InGameThreadLastTickTimeMs)
 	, GameThreadLastCompletionTimeMs(0.0f)
+#if ENABLE_DRAW_DEBUG
+	, PrevFrameGameThreadLastCompletionTimeMs(0.0f)
+#endif
 	, FrameOffset(0)
 	, DesiredTickRate(1)
 	, TickRate(1)
@@ -48,6 +50,7 @@ FAnimBudgetAllocatorComponentData::FAnimBudgetAllocatorComponentData(USkeletalMe
 	, bAllowReducedWork(true)
 	, bAutoCalculateSignificance(false)
 	, bOnScreen(false)
+	, bWasOffScreenPrevTick(false)
 	, bNeverThrottle(true)
 {}
 
@@ -180,6 +183,9 @@ void FAnimationBudgetAllocator::QueueSortedComponentIndices(float InDeltaSeconds
 	auto QueueComponentTick = [InDeltaSeconds, this, bUseSignificanceDelegate](FAnimBudgetAllocatorComponentData& InComponentData, int32 InComponentIndex, bool bInOnScreen)
 	{
 		InComponentData.AccumulatedDeltaTime += InDeltaSeconds;
+
+		InComponentData.bWasOffScreenPrevTick = (bInOnScreen && !InComponentData.bOnScreen);
+
 		InComponentData.bOnScreen = bInOnScreen;
 		InComponentData.StateChangeThrottle = InComponentData.StateChangeThrottle < 0 ? InComponentData.StateChangeThrottle : InComponentData.StateChangeThrottle - 1;
 
@@ -288,9 +294,20 @@ void FAnimationBudgetAllocator::QueueSortedComponentIndices(float InDeltaSeconds
 				else
 				{
 					DisableComponentTick(ComponentData);
+
+#if ENABLE_DRAW_DEBUG
+					NonRenderedNotSignificantComponentDataForDebugDraw.Add(ComponentIndex);
+#endif // #if ENABLE_DRAW_DEBUG
 				}
 			}
 		}
+#if ENABLE_DRAW_DEBUG
+		else
+		{
+			DisabledTickingComponentDataForDebugDraw.Add(ComponentIndex);
+		}
+#endif // #if ENABLE_DRAW_DEBUG
+
 
 		if(ComponentData.bReducedWork)
 		{
@@ -331,7 +348,11 @@ void FAnimationBudgetAllocator::QueueSortedComponentIndices(float InDeltaSeconds
 	{
 		auto ReduceWorkForOffscreenComponent = [](FAnimBudgetAllocatorComponentData& InComponentData)
 		{
-			if(InComponentData.bAllowReducedWork) 
+			if(InComponentData.bAllowReducedWork
+#if !(UE_BUILD_SHIPPING || UE_BUILD_TEST)
+				&& !GAnimationBudgetDebugForce
+#endif
+			)
 			{
 				InComponentData.bReducedWorkDueToOffscreen = true;
 
@@ -375,7 +396,7 @@ void FAnimationBudgetAllocator::QueueSortedComponentIndices(float InDeltaSeconds
 
 		if (GAnimationBudgetMaintainReduceWorkWhenOffScreen)
 		{
-			// if component ticks were disabled either here of through reduce work, they should be taken out of the sorted component list
+			// if component ticks were disabled either here or through reduce work, they should be taken out of the sorted component list
 			for (int32 SortedCompIdx = 0 ; SortedCompIdx < AllSortedComponentData.Num() ; ++SortedCompIdx)
 			{
 				FAnimBudgetAllocatorComponentData& SortedComponentData = AllComponentData[AllSortedComponentData[SortedCompIdx]];
@@ -405,6 +426,10 @@ int32 FAnimationBudgetAllocator::CalculateWorkDistributionAndQueue(float InDelta
 {
 	int32 NumTicked = 0;
 
+#if ENABLE_DRAW_DEBUG
+	bEmergencyReducedWorkActive = false;
+#endif
+
 	auto QueueForTick = [&NumTicked, this](FAnimBudgetAllocatorComponentData& InComponentData, int32 InStateChangeThrottleInFrames)
 	{
 		const int32 PrerequisiteHandle = InComponentData.RootPrerequisite != nullptr ? InComponentData.RootPrerequisite->GetAnimationBudgetHandle() : INDEX_NONE;
@@ -412,7 +437,18 @@ int32 FAnimationBudgetAllocator::CalculateWorkDistributionAndQueue(float InDelta
 
 		// Using (frame offset + frame counter) % tick rate allows us to only tick at the specified interval,
 		// but at a roughly even distribution over all registered components
-		const bool bTickThisFrame = (((GFrameCounter + ComponentDataToCheck.FrameOffset) % ComponentDataToCheck.TickRate) == 0);
+		bool bTickThisFrame = (((GFrameCounter + ComponentDataToCheck.FrameOffset) % ComponentDataToCheck.TickRate) == 0);
+
+		if (GAnimationBudgetForceTickWhenComponentExitsOffScreen != 0)
+		{
+			if (InComponentData.bWasOffScreenPrevTick && !bTickThisFrame)
+			{
+				// force tick this frame and fix up frameoffset so we maintain the same tick rate
+				bTickThisFrame = true;
+				InComponentData.FrameOffset += (InComponentData.TickRate - ((GFrameCounter + InComponentData.FrameOffset) % InComponentData.TickRate));
+			}
+		}
+
 		// Calculate interp alpha even when interpolation or ticking is disabled as this is used to decide how much root motion to consume
 		// each frame by character movement
 		InComponentData.SkippedTicks = bTickThisFrame ? 0 : (InComponentData.SkippedTicks + 1);
@@ -422,6 +458,10 @@ int32 FAnimationBudgetAllocator::CalculateWorkDistributionAndQueue(float InDelta
 		float Alpha = (InComponentData.TickRate > 1) ? FMath::Clamp((1.0f / TicksRemaining), 0.0f, 1.0f) : 1.f;
 		InComponentData.Component->SetExternalInterpolationAlpha(Alpha);
 
+#if ENABLE_DRAW_DEBUG
+		InComponentData.PrevFrameGameThreadLastCompletionTimeMs = InComponentData.GameThreadLastCompletionTimeMs;
+#endif
+
 		if((ComponentDataToCheck.bInterpolate && ComponentDataToCheck.bOnScreen) || bTickThisFrame)
 		{
 			InComponentData.bInterpolate = ComponentDataToCheck.bInterpolate;
@@ -466,23 +506,33 @@ int32 FAnimationBudgetAllocator::CalculateWorkDistributionAndQueue(float InDelta
 		for(int32 SortedComponentIndex : AllSortedComponentData)
 		{
 			FAnimBudgetAllocatorComponentData& ComponentData = AllComponentData[SortedComponentIndex];
-			ComponentData.TickRate = ComponentData.DesiredTickRate = GAnimationBudgetDebugForceRate;
-			ComponentData.bInterpolate = !!GAnimationBudgetDebugForceInterpolation;
-			if(ComponentData.Component->OnReduceWork().IsBound())
+
+			if (GAnimationBudgetDebugForceExcludeAlwaysTicked == 0 || !ComponentData.bAlwaysTick)
 			{
-				if(ComponentData.bReducedWork && !GAnimationBudgetDebugForceReducedWork)
-				{
-					ComponentData.Component->OnReduceWork().Execute(ComponentData.Component, false);
-					ComponentData.bReducedWork = false;
-				}
-				else if(!ComponentData.bReducedWork && GAnimationBudgetDebugForceReducedWork)
+				ComponentData.TickRate = ComponentData.DesiredTickRate = GAnimationBudgetDebugForceRate;
+				ComponentData.bInterpolate = !!GAnimationBudgetDebugForceInterpolation;
+				if (ComponentData.Component->OnReduceWork().IsBound())
 				{
-					ComponentData.Component->OnReduceWork().Execute(ComponentData.Component, true);
-					ComponentData.bReducedWork = true;
+					if (ComponentData.bReducedWork && !GAnimationBudgetDebugForceReducedWork)
+					{
+						ComponentData.Component->OnReduceWork().Execute(ComponentData.Component, false);
+						ComponentData.bReducedWork = false;
+					}
+					else if (!ComponentData.bReducedWork && GAnimationBudgetDebugForceReducedWork)
+					{
+						ComponentData.Component->OnReduceWork().Execute(ComponentData.Component, true);
+						ComponentData.bReducedWork = true;
+					}
 				}
 			}
-			ComponentData.StateChangeThrottle = -1;
+			else
+			{
+				ComponentData.TickRate = 1;
+				ComponentData.DesiredTickRate = 1;
+				ComponentData.bInterpolate = false;
+			}
 
+			ComponentData.StateChangeThrottle = -1;
 			QueueForTick(ComponentData, -1);
 		}
 	}
@@ -603,6 +653,12 @@ int32 FAnimationBudgetAllocator::CalculateWorkDistributionAndQueue(float InDelta
 
 			SET_FLOAT_STAT(STAT_AnimationBudgetAllocator_SmoothedBudgetPressure, SmoothedBudgetPressure);
 
+			const bool bEmergencyReducedWork = (SmoothedBudgetPressure >= Parameters.BudgetPressureBeforeEmergencyReducedWork) && (GAnimationBudgetEmergencyReduceWorkDisabled == 0);
+
+#if ENABLE_DRAW_DEBUG
+			bEmergencyReducedWorkActive = bEmergencyReducedWork;
+#endif
+
 			// Check all datas to line up tick rates based on prerequisites before queuing for ticks so ticks can be consistent
 			for (SortedComponentIndex = 0; SortedComponentIndex < TotalIdealWorkUnits; ++SortedComponentIndex)
 			{
@@ -646,8 +702,6 @@ int32 FAnimationBudgetAllocator::CalculateWorkDistributionAndQueue(float InDelta
 
 			if(--ReducedComponentWorkCounter <= 0)
 			{
-				const bool bEmergencyReducedWork = SmoothedBudgetPressure >= Parameters.BudgetPressureBeforeEmergencyReducedWork;
-
 				// Scale num components to switch based on budget pressure
 				const int32 NumComponentsToSwitch = (int32)FMath::Lerp(1.0f, (float)Parameters.ReducedWorkThrottleMaxPerFrame, BudgetPressureInterpAlpha);
 				int32 ComponentsSwitched = 0;
@@ -768,6 +822,11 @@ void FAnimationBudgetAllocator::Update(float DeltaSeconds)
 		AllSortedComponentDataDebug.Reset();
 #endif
 
+#if ENABLE_DRAW_DEBUG
+		DisabledTickingComponentDataForDebugDraw.Reset();
+		NonRenderedNotSignificantComponentDataForDebugDraw.Reset();
+#endif
+
 		QueueSortedComponentIndices(DeltaSeconds);
 
 		float AverageTickRate = 0.0f;
@@ -791,27 +850,162 @@ void FAnimationBudgetAllocator::Update(float DeltaSeconds)
 				ComponentIndexArray.Add(ComponentDataIndex);
 			}
 
-			for(const TPair<AActor*, TArray<int32>>& ActorIndicesPair : ActorMap)
+			if (GAnimationBudgetDebugVerboseActive)
 			{
-				FString DebugString;
-				
-				for(int32 ComponentDataIndex : ActorIndicesPair.Value)
+				const float WorldTime = World->TimeSeconds - 1.0f;
+
+				if (GAnimationBudgetDebugVerboseShowNonTicking)
 				{
-					FAnimBudgetAllocatorComponentData& ComponentData = AllComponentData[ComponentDataIndex];
-					if(ComponentData.bTickEnabled && ComponentData.bOnScreen)
+					for (int32 DisabledTickingComponentDataIndex : DisabledTickingComponentDataForDebugDraw)
+					{
+						FAnimBudgetAllocatorComponentData& ComponentData = AllComponentData[DisabledTickingComponentDataIndex];
+						TArray<int32>& ComponentIndexArray = ActorMap.FindOrAdd(ComponentData.Component->GetOwner());
+						ComponentIndexArray.Add(DisabledTickingComponentDataIndex);
+					}
+
+					for (int32 NonRenderedComponentDataIndex : NonRenderedNotSignificantComponentDataForDebugDraw)
+					{
+						FAnimBudgetAllocatorComponentData& ComponentData = AllComponentData[NonRenderedComponentDataIndex];
+						TArray<int32>& ComponentIndexArray = ActorMap.FindOrAdd(ComponentData.Component->GetOwner());
+						ComponentIndexArray.Add(NonRenderedComponentDataIndex);
+					}
+				}
+
+				for(TPair<AActor*, TArray<int32>>& ActorIndicesPair : ActorMap)
+				{
+					ActorIndicesPair.Value.Sort();
+
+					FString DebugString;
+					bool bActorOffscreen = false;
+				
+					for(int32 ComponentDataIndex : ActorIndicesPair.Value)
 					{
-						if(GAnimationBudgetDebugShowAddresses != 0)
+						FAnimBudgetAllocatorComponentData& ComponentData = AllComponentData[ComponentDataIndex];
+
+						if (!ComponentData.bOnScreen)
 						{
-							DebugString += FString::Printf(TEXT("0x%" UPTRINT_x_FMT " %d %s %s\n"), (UPTRINT)&ComponentData, ComponentData.TickRate, ComponentData.bInterpolate ? TEXT("I") : TEXT(" "), ComponentData.bReducedWork ? TEXT("Lo") : TEXT("Hi"));
+							bActorOffscreen = true;
 						}
-						else
+
+						// component could be not ticking due to tick being disabled externally (bTickEnabled is false) or due to being not visible, etc
+						const bool bTicking = ComponentData.Component->PrimaryComponentTick.IsTickFunctionEnabled();
+
+						if((bTicking && ComponentData.bOnScreen) || GAnimationBudgetDebugVerboseShowNonTicking)
 						{
-							DebugString += FString::Printf(TEXT("%.03f %d %s %s\n"), ComponentData.Significance, ComponentData.TickRate, ComponentData.bInterpolate ? TEXT("I") : TEXT(" "), ComponentData.bReducedWork ? TEXT("Lo") : TEXT("Hi"));
+							FString StatusStr = FString::Printf(TEXT("%s %s %.2f "),
+								ComponentData.Component->bHiddenInGame ? TEXT("HID") : TEXT("NOTHID"),
+								ComponentData.Component->IsVisible() ? TEXT("VIS") : TEXT("NOTVIS"), 
+								ComponentData.Component->GetComponentScale().X);
+
+							if (!ComponentData.bTickEnabled)
+							{
+								StatusStr += TEXT(" DISABLED ");
+							}
+							else if (NonRenderedNotSignificantComponentDataForDebugDraw.Contains(ComponentDataIndex))
+							{
+								StatusStr += TEXT(" OFF (NOT SIG) ");
+							}
+							else if (!ComponentData.bOnScreen)
+							{
+								StatusStr += TEXT(" OFF (SIG) ");
+							}
+
+							FString TickStr = TEXT("   ");
+
+							if (bTicking)
+							{
+								// if offscreen 
+								if (ComponentData.Component->GetLastRenderTime() <= WorldTime)
+								{
+									switch (ComponentData.Component->VisibilityBasedAnimTickOption)
+									{
+										case EVisibilityBasedAnimTickOption::AlwaysTickPose:
+										case EVisibilityBasedAnimTickOption::AlwaysTickPoseAndRefreshBones:
+										{
+											TickStr = TEXT(" A");
+										}
+										break;
+
+										case EVisibilityBasedAnimTickOption::OnlyTickMontagesWhenNotRendered:
+										{
+											TickStr = TEXT(" M");
+										}
+										break;
+
+										case EVisibilityBasedAnimTickOption::OnlyTickPoseWhenRendered:
+										{
+											TickStr = TEXT(" X");
+										}
+										break;
+
+										default:
+										break;
+									}
+								}
+								else
+								{
+									TickStr = TEXT(" *");
+								}
+							}
+
+							DebugString += FString::Printf(TEXT("%s%s%s%.03f %d %s %s%s - T:%.3f E:%.3f A:%.3f\n"), *StatusStr,
+								(GAnimationBudgetDebugShowAddresses) ? *FString::Printf(TEXT("0x%" UPTRINT_x_FMT " "), (UPTRINT)ComponentData.Component) : TEXT(""),
+								(GAnimationBudgetDebugVerboseShowNames) ? *FString::Printf(TEXT("%s%s "), (GAnimationBudgetDebugVerboseShowActorNames) ? *FString::Printf(TEXT("%s - "), *ComponentData.Component->GetOwner()->GetName()) : TEXT(""), *ComponentData.Component->GetName()) : TEXT(""),
+								ComponentData.Significance, ComponentData.TickRate, ComponentData.bInterpolate ? TEXT("I") : TEXT(" "), ComponentData.bReducedWork ? TEXT("Lo") : TEXT("Hi"), *TickStr,
+								ComponentData.GameThreadLastTickTimeMs, ComponentData.PrevFrameGameThreadLastCompletionTimeMs, ComponentData.GameThreadLastTickTimeMs + ComponentData.PrevFrameGameThreadLastCompletionTimeMs);
 						}
 					}
+				
+					if (GAnimationBudgetDebugVerboseOnCanvas != 0)
+					{
+						GEngine->AddOnScreenDebugMessage(-1, 0.0f, FColor::White, DebugString);
+					}
+					else
+					{
+						DrawDebugString(World, FVector::ZeroVector, DebugString, ActorIndicesPair.Key, FColor::White, 0.016f, false);
+					}
+				}
+
+				if (bEmergencyReducedWorkActive)
+				{
+					GEngine->AddOnScreenDebugMessage(-1, 0.0f, FColor::Red, TEXT("Emergency Reduce Work Enabled"));
 				}
 
-				DrawDebugString(World, FVector::ZeroVector, DebugString, ActorIndicesPair.Key, FColor::White, 0.016f, false);
+				if (AllSortedComponentData.Num() > 0)
+				{
+					FString Prefix;
+					if (World->WorldType == EWorldType::PIE && World->GetNetMode() == NM_Client)
+					{
+						Prefix = FString::Printf(TEXT("Client %d: "), UE::GetPlayInEditorID());
+					}
+
+					GEngine->AddOnScreenDebugMessage(-1, 0.0f, FColor::White, *FString::Printf(TEXT("%s Frame: %d Frame Budget: %.4f"), *Prefix, GFrameNumber, Parameters.BudgetInMs));
+				}
+			}
+			else  // !GAnimationBudgetDebugVerboseActive
+			{
+				for (const TPair<AActor*, TArray<int32>>& ActorIndicesPair : ActorMap)
+				{
+					FString DebugString;
+
+					for (int32 ComponentDataIndex : ActorIndicesPair.Value)
+					{
+						FAnimBudgetAllocatorComponentData& ComponentData = AllComponentData[ComponentDataIndex];
+						if (ComponentData.bTickEnabled && ComponentData.bOnScreen)
+						{
+							if (GAnimationBudgetDebugShowAddresses != 0)
+							{
+								DebugString += FString::Printf(TEXT("0x%" UPTRINT_x_FMT " %d %s %s\n"), (UPTRINT)&ComponentData, ComponentData.TickRate, ComponentData.bInterpolate ? TEXT("I") : TEXT(" "), ComponentData.bReducedWork ? TEXT("Lo") : TEXT("Hi"));
+							}
+							else
+							{
+								DebugString += FString::Printf(TEXT("%.03f %d %s %s\n"), ComponentData.Significance, ComponentData.TickRate, ComponentData.bInterpolate ? TEXT("I") : TEXT(" "), ComponentData.bReducedWork ? TEXT("Lo") : TEXT("Hi"));
+							}
+						}
+					}
+
+					DrawDebugString(World, FVector::ZeroVector, DebugString, ActorIndicesPair.Key, FColor::White, 0.016f, false);
+				}
 			}
 
 			DebugTimes.Add(FVector2D(CurrentDebugTimeDisplay, DebugTotalTime));
@@ -836,7 +1030,7 @@ void FAnimationBudgetAllocator::Update(float DeltaSeconds)
 #if ENABLE_DRAW_DEBUG
 void FAnimationBudgetAllocator::OnHUDPostRender(AHUD* HUD, UCanvas* Canvas)
 {
-	if(FAnimationBudgetAllocator::bCachedEnabled && bEnabled && GAnimationBudgetDebugEnabled != 0)
+	if(FAnimationBudgetAllocator::bCachedEnabled && bEnabled && GAnimationBudgetDebugEnabled != 0 && GAnimationBudgetDebugDisableGraph == 0)
 	{
 		if(HUD->GetWorld() == World)
 		{
diff --git a/Engine/Plugins/Runtime/AnimationBudgetAllocator/Source/AnimationBudgetAllocator/Private/AnimationBudgetAllocator.h b/Engine/Plugins/Runtime/AnimationBudgetAllocator/Source/AnimationBudgetAllocator/Private/AnimationBudgetAllocator.h
--- a/Engine/Plugins/Runtime/AnimationBudgetAllocator/Source/AnimationBudgetAllocator/Private/AnimationBudgetAllocator.h
+++ b/Engine/Plugins/Runtime/AnimationBudgetAllocator/Source/AnimationBudgetAllocator/Private/AnimationBudgetAllocator.h
@@ -41,6 +41,9 @@ struct FAnimBudgetAllocatorComponentData
 		, AccumulatedDeltaTime(0.0f)
 		, GameThreadLastTickTimeMs(0.0f)
 		, GameThreadLastCompletionTimeMs(0.0f)
+#if ENABLE_DRAW_DEBUG
+		, PrevFrameGameThreadLastCompletionTimeMs(0.0f)
+#endif
 		, FrameOffset(0)
 		, DesiredTickRate(0)
 		, TickRate(0)		
@@ -55,6 +58,7 @@ struct FAnimBudgetAllocatorComponentData
 		, bAllowReducedWork(false)
 		, bAutoCalculateSignificance(false)
 		, bOnScreen(false)
+		, bWasOffScreenPrevTick(false)
 		, bNeverThrottle(true)
 	{}
 
@@ -84,6 +88,11 @@ public:
 	/** Tracks the time in MS it took to complete this component on the game thread */
 	float GameThreadLastCompletionTimeMs;
 
+#if ENABLE_DRAW_DEBUG
+	/** Keep prev frame GameThreadLastCompletionTimeMs for debug use */
+	float PrevFrameGameThreadLastCompletionTimeMs;
+#endif
+
 	/** Frame offset used to distribute ticks */
 	uint32 FrameOffset;
 
@@ -126,6 +135,9 @@ public:
 	/** Whether this component is on screen. This is updated each tick. */
 	uint8 bOnScreen : 1;
 
+	/** Whether this component was off screen prev tick. */
+	uint8 bWasOffScreenPrevTick : 1;
+
 	/** Whether we are allowing interpolation on this component (i.e. we dont just reduce tick rate). This is intended to allow higher-quality animation. */
 	uint8 bNeverThrottle : 1;
 };
@@ -230,6 +242,11 @@ protected:
 	TArray<FAnimBudgetAllocatorComponentData*> AllSortedComponentDataDebug;
 #endif
 
+#if ENABLE_DRAW_DEBUG
+	TArray<int32> DisabledTickingComponentDataForDebugDraw;
+	TArray<int32> NonRenderedNotSignificantComponentDataForDebugDraw;
+#endif
+
 	/** All components that have reduced work that might want to tick (and hence might not want to do reduced work) */
 	TArray<int32> ReducedWorkComponentData;
 
@@ -316,6 +333,9 @@ protected:
 	/** External animation time spent by sources not captured in ABA but that should reduce avaialbe ABA budget for the frame */
 	float ExternalAnimationTimeSpent = 0.0f;
 
+#if ENABLE_DRAW_DEBUG
+	bool bEmergencyReducedWorkActive = false;
+#endif
 };
 
 #undef UE_API
diff --git a/Engine/Plugins/Runtime/AnimationBudgetAllocator/Source/AnimationBudgetAllocator/Private/AnimationBudgetAllocatorCVars.cpp b/Engine/Plugins/Runtime/AnimationBudgetAllocator/Source/AnimationBudgetAllocator/Private/AnimationBudgetAllocatorCVars.cpp
--- a/Engine/Plugins/Runtime/AnimationBudgetAllocator/Source/AnimationBudgetAllocator/Private/AnimationBudgetAllocatorCVars.cpp
+++ b/Engine/Plugins/Runtime/AnimationBudgetAllocator/Source/AnimationBudgetAllocator/Private/AnimationBudgetAllocatorCVars.cpp
@@ -6,18 +6,18 @@
 
 int32 GAnimationBudgetEnabled = 0;
 
-static FAutoConsoleVariableRef CVarSkelBatch_Enabled(
+static FAutoConsoleVariableRef CVarAnimBudget_Enabled(
 	TEXT("a.Budget.Enabled"),
 	GAnimationBudgetEnabled,
 	TEXT("Values: 0/1\n")
-	TEXT("Controls whether the skeletal mesh batching system is enabled. Should be set when there are no running skeletal meshes."),
+	TEXT("Controls whether the Animation Budgeting system is enabled. Should be set when there are no running skeletal meshes."),
 	ECVF_Scalability);
 
 // some optional changes to behavior in the budgeter
 
 int32 GAnimationBudgetMaintainReduceWorkWhenOffScreen = 0;
 
-static FAutoConsoleVariableRef CVarSkelBatch_MaintainReduceWorkWhenOffScreen(
+static FAutoConsoleVariableRef CVarAnimBudget_MaintainReduceWorkWhenOffScreen(
 	TEXT("a.Budget.MaintainReduceWorkWhenOffScreen"),
 	GAnimationBudgetMaintainReduceWorkWhenOffScreen,
 	TEXT("Values: 0/1\n")
@@ -26,38 +26,112 @@ static FAutoConsoleVariableRef CVarSkelBatch_MaintainReduceWorkWhenOffScreen(
 
 int32 GAnimationBudgetComponentsWithZeroSigAreNonRendered = 0;
 
-static FAutoConsoleVariableRef CVarSkelBatch_ComponentsWithZeroSigAreNonRendered(
+static FAutoConsoleVariableRef CVarAnimBudget_ComponentsWithZeroSigAreNonRendered(
 	TEXT("a.Budget.ComponentsWithZeroSigAreNonRendered"),
 	GAnimationBudgetComponentsWithZeroSigAreNonRendered,
 	TEXT("Values: 0/1\n")
 	TEXT("Control whether components with zero significance can be treated in the animation budgeter as non-rendered."),
 	ECVF_Scalability);
 
+int32 GAnimationBudgetForceTickWhenComponentExitsOffScreen = 0;
+
+static FAutoConsoleVariableRef CVarAnimBudget_ForceTickWhenComponentExitsOffScreen(
+	TEXT("a.Budget.ForceTickWhenComponentExitsOffScreen"),
+	GAnimationBudgetForceTickWhenComponentExitsOffScreen,
+	TEXT("Values: 0/1\n")
+	TEXT("Controls whether components not ticking every frame have a forced tick when exiting offscreen."),
+	ECVF_Scalability);
+
+int32 GAnimationBudgetEmergencyReduceWorkDisabled = 0;
+
+static FAutoConsoleVariableRef CVarAnimBudget_EmergencyReduceWorkDisabled(
+	TEXT("a.Budget.EmergencyReduceWorkDisabled"),
+	GAnimationBudgetEmergencyReduceWorkDisabled,
+	TEXT("Values: 0/1\n")
+	TEXT("Controls whether the emergency reduce work option is disabled."),
+	ECVF_Scalability);
+
 #if ENABLE_DRAW_DEBUG
+
 int32 GAnimationBudgetDebugEnabled = 0;
 
-static FAutoConsoleVariableRef CVarSkelBatch_DebugEnabled(
+static FAutoConsoleVariableRef CVarAnimBudget_DebugEnabled(
 	TEXT("a.Budget.Debug.Enabled"),
 	GAnimationBudgetDebugEnabled,
 	TEXT("Values: 0/1\n")
 	TEXT("Controls whether debug rendering (in builds that support it) is enabled for the animation budget allocator."),
 	ECVF_Scalability);
 
+int32 GAnimationBudgetDebugDisableGraph = 0;
+
+static FAutoConsoleVariableRef CVarAnimBudget_DisableGraph(
+	TEXT("a.Budget.Debug.DisableGraph"),
+	GAnimationBudgetDebugDisableGraph,
+	TEXT("Values: 0/1\n")
+	TEXT("Controls whether debug rendering of perf graph is disabled or not."),
+	ECVF_Scalability);
+
 int32 GAnimationBudgetDebugShowAddresses = 0;
 
-static FAutoConsoleVariableRef CVarSkelBatch_ShowAddresses(
+static FAutoConsoleVariableRef CVarAnimBudget_ShowAddresses(
 	TEXT("a.Budget.Debug.ShowAddresses"),
 	GAnimationBudgetDebugShowAddresses,
 	TEXT("Values: 0/1\n")
 	TEXT("Controls whether debug rendering shows addresses of component data for debugging."),
 	ECVF_Scalability);
+
+int32 GAnimationBudgetDebugVerboseActive = 0;
+
+static FAutoConsoleVariableRef CVarAnimBudget_VerboseActive(
+	TEXT("a.Budget.Debug.VerboseActive"),
+	GAnimationBudgetDebugVerboseActive,
+	TEXT("Values: 0/1\n")
+	TEXT("Controls whether debug rendering is in verbose mode and shows a lot of extra information."),
+	ECVF_Scalability);
+
+int32 GAnimationBudgetDebugVerboseShowNames = 0;
+
+static FAutoConsoleVariableRef CVarAnimBudget_VerboseShowNames(
+	TEXT("a.Budget.Debug.Verbose.ShowNames"),
+	GAnimationBudgetDebugVerboseShowNames,
+	TEXT("Values: 0/1\n")
+	TEXT("In Verbose Mode - Controls whether debug rendering shows names of components for debugging."),
+	ECVF_Scalability);
+
+int32 GAnimationBudgetDebugVerboseShowActorNames = 0;
+
+static FAutoConsoleVariableRef CVarAnimBudget_VerboseShowActorNames(
+	TEXT("a.Budget.Debug.Verbose.ShowActorNames"),
+	GAnimationBudgetDebugVerboseShowActorNames,
+	TEXT("Values: 0/1\n")
+	TEXT("In Verbose Mode - Controls whether debug rendering shows actor names for debugging."),
+	ECVF_Scalability);
+
+int32 GAnimationBudgetDebugVerboseShowNonTicking = 0;
+
+static FAutoConsoleVariableRef CVarAnimBudget_VerboseShowNonTicking(
+	TEXT("a.Budget.Debug.Verbose.ShowNonTicking"),
+	GAnimationBudgetDebugVerboseShowNonTicking,
+	TEXT("Values: 0/1\n")
+	TEXT("In Verbose Mode - Controls whether debug rendering shows non ticking component datas."),
+	ECVF_Scalability);
+
+int32 GAnimationBudgetDebugVerboseOnCanvas = 0;
+
+static FAutoConsoleVariableRef CVarAnimBudget_VerboseShowOnCanvas(
+	TEXT("a.Budget.Debug.Verbose.ShowOnCanvas"),
+	GAnimationBudgetDebugVerboseOnCanvas,
+	TEXT("Values: 0/1\n")
+	TEXT("In Verbose Mode - Controls whether debug rendering should be shown as 2D on canvas or not."),
+	ECVF_Scalability);
+
 #endif
 
 #if !(UE_BUILD_SHIPPING || UE_BUILD_TEST)
 
 int32 GAnimationBudgetDebugForce = 0;
 
-static FAutoConsoleVariableRef CVarSkelBatch_Force(
+static FAutoConsoleVariableRef CVarAnimBudget_Force(
 	TEXT("a.Budget.Debug.Force"),
 	GAnimationBudgetDebugForce,
 	TEXT("Values: 0/1\n")
@@ -65,9 +139,19 @@ static FAutoConsoleVariableRef CVarSkelBatch_Force(
 	TEXT("Turns on forced rate/interp/reduced controls. These override any budget-driven values."),
 	ECVF_Scalability);
 
+int32 GAnimationBudgetDebugForceExcludeAlwaysTicked = 0;
+
+static FAutoConsoleVariableRef CVarAnimBudget_ForceExcludeAlwaysTicked(
+	TEXT("a.Budget.Debug.Force.ExcludeAlwaysTicked"),
+	GAnimationBudgetDebugForceExcludeAlwaysTicked,
+	TEXT("Values: 0/1\n")
+	TEXT("Default: 0\n")
+	TEXT("Excludes any components set to always tick from having forced rate/interp/reduced controls."),
+	ECVF_Scalability);
+
 int32 GAnimationBudgetDebugForceRate = 4;
 
-static FAutoConsoleVariableRef CVarSkelBatch_ForceRate(
+static FAutoConsoleVariableRef CVarAnimBudget_ForceRate(
 	TEXT("a.Budget.Debug.Force.Rate"),
 	GAnimationBudgetDebugForceRate,
 	TEXT("Values: >= 1\n")
@@ -81,7 +165,7 @@ static FAutoConsoleVariableRef CVarSkelBatch_ForceRate(
 
 int32 GAnimationBudgetDebugForceInterpolation = 0;
 
-static FAutoConsoleVariableRef CVarSkelBatch_ForceInterpolation(
+static FAutoConsoleVariableRef CVarAnimBudget_ForceInterpolation(
 	TEXT("a.Budget.Debug.Force.Interp"),
 	GAnimationBudgetDebugForceInterpolation,
 	TEXT("Values: 0/1\n")
@@ -91,7 +175,7 @@ static FAutoConsoleVariableRef CVarSkelBatch_ForceInterpolation(
 
 int32 GAnimationBudgetDebugForceReducedWork = 0;
 
-static FAutoConsoleVariableRef CVarSkelBatch_ForceReducedWork(
+static FAutoConsoleVariableRef CVarAnimBudget_ForceReducedWork(
 	TEXT("a.Budget.Debug.Force.Reduced"),
 	GAnimationBudgetDebugForceReducedWork,
 	TEXT("Values: 0/1\n")
@@ -107,7 +191,7 @@ FAnimationBudgetAllocatorParameters GBudgetParameters;
 /** Delegate broadcast when parameter block changes */
 FSimpleMulticastDelegate GOnCVarParametersChanged;
 
-static FAutoConsoleVariableRef CVarSkelBatch_Budget(
+static FAutoConsoleVariableRef CVarAnimBudget_Budget(
 	TEXT("a.Budget.BudgetMs"),
 	GBudgetParameters.BudgetInMs,
 	TEXT("Values > 0.1, Default = 1.0\n")
@@ -120,7 +204,7 @@ static FAutoConsoleVariableRef CVarSkelBatch_Budget(
 	}),
 	ECVF_Scalability);
 
-static FAutoConsoleVariableRef CVarSkelBatch_MinQuality(
+static FAutoConsoleVariableRef CVarAnimBudget_MinQuality(
 	TEXT("a.Budget.MinQuality"),
 	GBudgetParameters.MinQuality,
 	TEXT("Values [0.0, 1.0], Default = 0.0\n")
@@ -133,7 +217,7 @@ static FAutoConsoleVariableRef CVarSkelBatch_MinQuality(
 	}),
 	ECVF_Scalability);
 
-static FAutoConsoleVariableRef CVarSkelBatch_MaxTickRate(
+static FAutoConsoleVariableRef CVarAnimBudget_MaxTickRate(
 	TEXT("a.Budget.MaxTickRate"),
 	GBudgetParameters.MaxTickRate,
 	TEXT("Values >= 1, Default = 10\n")
@@ -145,7 +229,7 @@ static FAutoConsoleVariableRef CVarSkelBatch_MaxTickRate(
 	}),
 	ECVF_Scalability);
 
-static FAutoConsoleVariableRef CVarSkelBatch_WorkUnitSmoothingSpeed(
+static FAutoConsoleVariableRef CVarAnimBudget_WorkUnitSmoothingSpeed(
 	TEXT("a.Budget.WorkUnitSmoothingSpeed"),
 	GBudgetParameters.WorkUnitSmoothingSpeed,
 	TEXT("Values > 0.1, Default = 5.0\n")
@@ -156,7 +240,7 @@ static FAutoConsoleVariableRef CVarSkelBatch_WorkUnitSmoothingSpeed(
 		GOnCVarParametersChanged.Broadcast();
 	}));
 
-static FAutoConsoleVariableRef CVarSkelBatch_AlwaysTickFalloffAggression(
+static FAutoConsoleVariableRef CVarAnimBudget_AlwaysTickFalloffAggression(
 	TEXT("a.Budget.AlwaysTickFalloffAggression"),
 	GBudgetParameters.AlwaysTickFalloffAggression,
 	TEXT("Range [0.1, 0.9], Default = 0.8\n")
@@ -169,7 +253,7 @@ static FAutoConsoleVariableRef CVarSkelBatch_AlwaysTickFalloffAggression(
 	}),
 	ECVF_Scalability);
 
-static FAutoConsoleVariableRef CVarSkelBatch_InterpolationFalloffAggression(
+static FAutoConsoleVariableRef CVarAnimBudget_InterpolationFalloffAggression(
 	TEXT("a.Budget.InterpolationFalloffAggression"),
 	GBudgetParameters.InterpolationFalloffAggression,
 	TEXT("Range [0.1, 0.9], Default = 0.4\n")
@@ -183,7 +267,7 @@ static FAutoConsoleVariableRef CVarSkelBatch_InterpolationFalloffAggression(
 	}),
 	ECVF_Scalability);
 
-static FAutoConsoleVariableRef CVarSkelBatch_InterpolationMaxRate(
+static FAutoConsoleVariableRef CVarAnimBudget_InterpolationMaxRate(
 	TEXT("a.Budget.InterpolationMaxRate"),
 	GBudgetParameters.InterpolationMaxRate,
 	TEXT("Values > 1, Default = 6\n")
@@ -195,7 +279,7 @@ static FAutoConsoleVariableRef CVarSkelBatch_InterpolationMaxRate(
 	}),
 	ECVF_Scalability);
 
-static FAutoConsoleVariableRef CVarSkelBatch_MaxInterpolatedComponents(
+static FAutoConsoleVariableRef CVarAnimBudget_MaxInterpolatedComponents(
 	TEXT("a.Budget.MaxInterpolatedComponents"),
 	GBudgetParameters.MaxInterpolatedComponents,
 	TEXT("Range >= 0, Default = 16\n")
@@ -207,7 +291,7 @@ static FAutoConsoleVariableRef CVarSkelBatch_MaxInterpolatedComponents(
 	}),
 	ECVF_Scalability);
 
-static FAutoConsoleVariableRef CVarSkelBatch_InterpolationTickMultiplier(
+static FAutoConsoleVariableRef CVarAnimBudget_InterpolationTickMultiplier(
 	TEXT("a.Budget.InterpolationTickMultiplier"),
 	GBudgetParameters.InterpolationTickMultiplier,
 	TEXT("Range [0.1, 0.9], Default = 0.75\n")
@@ -219,7 +303,7 @@ static FAutoConsoleVariableRef CVarSkelBatch_InterpolationTickMultiplier(
 	}),
 	ECVF_Scalability);
 
-static FAutoConsoleVariableRef CVarSkelBatch_InitialEstimatedWorkUnitTime(
+static FAutoConsoleVariableRef CVarAnimBudget_InitialEstimatedWorkUnitTime(
 	TEXT("a.Budget.InitialEstimatedWorkUnitTime"),
 	GBudgetParameters.InitialEstimatedWorkUnitTimeMs,
 	TEXT("Values > 0.0, Default = 0.08\n")
@@ -232,7 +316,7 @@ static FAutoConsoleVariableRef CVarSkelBatch_InitialEstimatedWorkUnitTime(
 	}),
 	ECVF_Scalability);
 
-static FAutoConsoleVariableRef CVarSkelBatch_MaxTickedOffsreenComponents(
+static FAutoConsoleVariableRef CVarAnimBudget_MaxTickedOffsreenComponents(
 	TEXT("a.Budget.MaxTickedOffsreen"),
 	GBudgetParameters.MaxTickedOffsreenComponents,
 	TEXT("Values >= 1, Default = 4\n")
@@ -244,7 +328,7 @@ static FAutoConsoleVariableRef CVarSkelBatch_MaxTickedOffsreenComponents(
 	}),
 	ECVF_Scalability);
 
-static FAutoConsoleVariableRef CVarSkelBatch_StateChangeThrottleInFrames(
+static FAutoConsoleVariableRef CVarAnimBudget_StateChangeThrottleInFrames(
 	TEXT("a.Budget.StateChangeThrottleInFrames"),
 	GBudgetParameters.StateChangeThrottleInFrames,
 	TEXT("Range [1, 128], Default = 30\n")
@@ -257,7 +341,7 @@ static FAutoConsoleVariableRef CVarSkelBatch_StateChangeThrottleInFrames(
 	ECVF_Scalability);
 
 
-static FAutoConsoleVariableRef CVarSkelBatch_BudgetFactorBeforeReducedWork(
+static FAutoConsoleVariableRef CVarAnimBudget_BudgetFactorBeforeReducedWork(
 	TEXT("a.Budget.BudgetFactorBeforeReducedWork"),
 	GBudgetParameters.BudgetFactorBeforeReducedWork,
 	TEXT("Range > 1, Default = 1.5\n")
@@ -269,7 +353,7 @@ static FAutoConsoleVariableRef CVarSkelBatch_BudgetFactorBeforeReducedWork(
 	}),
 	ECVF_Scalability);
 
-static FAutoConsoleVariableRef CVarSkelBatch_BudgetFactorBeforeReducedWorkEpsilon(
+static FAutoConsoleVariableRef CVarAnimBudget_BudgetFactorBeforeReducedWorkEpsilon(
 	TEXT("a.Budget.BudgetFactorBeforeReducedWorkEpsilon"),
 	GBudgetParameters.BudgetFactorBeforeReducedWorkEpsilon,
 	TEXT("Range > 0.0, Default = 0.25\n")
@@ -281,7 +365,7 @@ static FAutoConsoleVariableRef CVarSkelBatch_BudgetFactorBeforeReducedWorkEpsilo
 	}),
 	ECVF_Scalability);
 
-static FAutoConsoleVariableRef CVarSkelBatch_BudgetPressureSmoothingSpeed(
+static FAutoConsoleVariableRef CVarAnimBudget_BudgetPressureSmoothingSpeed(
 	TEXT("a.Budget.BudgetPressureSmoothingSpeed"),
 	GBudgetParameters.BudgetPressureSmoothingSpeed,
 	TEXT("Range > 0.0, Default = 3.0\n")
@@ -293,7 +377,7 @@ static FAutoConsoleVariableRef CVarSkelBatch_BudgetPressureSmoothingSpeed(
 	}),
 	ECVF_Scalability);
 
-static FAutoConsoleVariableRef CVarSkelBatch_ReducedWorkThrottleMinInFrames(
+static FAutoConsoleVariableRef CVarAnimBudget_ReducedWorkThrottleMinInFrames(
 	TEXT("a.Budget.ReducedWorkThrottleMinInFrames"),
 	GBudgetParameters.ReducedWorkThrottleMinInFrames,
 	TEXT("Range [1, 255], Default = 2\n")
@@ -305,7 +389,7 @@ static FAutoConsoleVariableRef CVarSkelBatch_ReducedWorkThrottleMinInFrames(
 	}),
 	ECVF_Scalability);
 
-static FAutoConsoleVariableRef CVarSkelBatch_ReducedWorkThrottleMaxInFrames(
+static FAutoConsoleVariableRef CVarAnimBudget_ReducedWorkThrottleMaxInFrames(
 	TEXT("a.Budget.ReducedWorkThrottleMaxInFrames"),
 	GBudgetParameters.ReducedWorkThrottleMaxInFrames,
 	TEXT("Range [1, 255], Default = 20\n")
@@ -317,7 +401,7 @@ static FAutoConsoleVariableRef CVarSkelBatch_ReducedWorkThrottleMaxInFrames(
 	}),
 	ECVF_Scalability);
 
-static FAutoConsoleVariableRef CVarSkelBatch_BudgetFactorBeforeAggressiveReducedWork(
+static FAutoConsoleVariableRef CVarAnimBudget_BudgetFactorBeforeAggressiveReducedWork(
 	TEXT("a.Budget.BudgetFactorBeforeAggressiveReducedWork"),
 	GBudgetParameters.BudgetFactorBeforeAggressiveReducedWork,
 	TEXT("Range > 1, Default = 2.0\n")
@@ -329,7 +413,7 @@ static FAutoConsoleVariableRef CVarSkelBatch_BudgetFactorBeforeAggressiveReduced
 	}),
 	ECVF_Scalability);
 
-static FAutoConsoleVariableRef CVarSkelBatch_ReducedWorkThrottleMaxPerFrame(
+static FAutoConsoleVariableRef CVarAnimBudget_ReducedWorkThrottleMaxPerFrame(
 	TEXT("a.Budget.ReducedWorkThrottleMaxPerFrame"),
 	GBudgetParameters.ReducedWorkThrottleMaxPerFrame,
 	TEXT("Range [1, 255], Default = 4\n")
@@ -341,8 +425,8 @@ static FAutoConsoleVariableRef CVarSkelBatch_ReducedWorkThrottleMaxPerFrame(
 	}),
 	ECVF_Scalability);
 
-static FAutoConsoleVariableRef CVarSkelBatch_BudgetPressureBeforeEmergencyReducedWork(
-	TEXT("a.Budget.GBudgetPressureBeforeEmergencyReducedWork"),
+static FAutoConsoleVariableRef CVarAnimBudget_BudgetPressureBeforeEmergencyReducedWork(
+	TEXT("a.Budget.BudgetPressureBeforeEmergencyReducedWork"),
 	GBudgetParameters.BudgetPressureBeforeEmergencyReducedWork,
 	TEXT("Range > 0.0, Default = 2.5\n")
 	TEXT("Controls the budget pressure where emergency reduced work (applied to all components except those that are bAlwaysTick).\n"),
@@ -353,7 +437,7 @@ static FAutoConsoleVariableRef CVarSkelBatch_BudgetPressureBeforeEmergencyReduce
 	}),
 	ECVF_Scalability);
 
-static FAutoConsoleVariableRef CVarSkelBatch_AutoCalculatedSignificanceMaxDistance(
+static FAutoConsoleVariableRef CVarAnimBudget_AutoCalculatedSignificanceMaxDistance(
 	TEXT("a.Budget.AutoCalculatedSignificanceMaxDistance"),
 	GBudgetParameters.AutoCalculatedSignificanceMaxDistance,
 	TEXT("Range > 1.0, Default = 300.0\n")
diff --git a/Engine/Plugins/Runtime/AnimationBudgetAllocator/Source/AnimationBudgetAllocator/Private/AnimationBudgetAllocatorCVars.h b/Engine/Plugins/Runtime/AnimationBudgetAllocator/Source/AnimationBudgetAllocator/Private/AnimationBudgetAllocatorCVars.h
--- a/Engine/Plugins/Runtime/AnimationBudgetAllocator/Source/AnimationBudgetAllocator/Private/AnimationBudgetAllocatorCVars.h
+++ b/Engine/Plugins/Runtime/AnimationBudgetAllocator/Source/AnimationBudgetAllocator/Private/AnimationBudgetAllocatorCVars.h
@@ -13,18 +13,44 @@ extern int32 GAnimationBudgetMaintainReduceWorkWhenOffScreen;
 
 extern int32 GAnimationBudgetComponentsWithZeroSigAreNonRendered;
 
+extern int32 GAnimationBudgetForceTickWhenComponentExitsOffScreen;
+
+/* Controls whether the emergency reduce work option is disabled. */
+extern int32 GAnimationBudgetEmergencyReduceWorkDisabled;
+
 #if ENABLE_DRAW_DEBUG
+
 /** Debug rendering flag */
 extern int32 GAnimationBudgetDebugEnabled;
 
+/* Controls whether debug rendering of perf graph is disabled or not. */
+extern int32 GAnimationBudgetDebugDisableGraph;
+
 /** Controls whether debug rendering shows addresses of component data for debugging */
 extern int32 GAnimationBudgetDebugShowAddresses;
+
+/** Controls whether debug rendering is in verbose mode and shows a lot of extra information */
+extern int32 GAnimationBudgetDebugVerboseActive;
+
+/* In Verbose Mode - Controls whether debug rendering shows names of components for debugging. */
+extern int32 GAnimationBudgetDebugVerboseShowNames;
+
+/* In Verbose Mode - Controls whether debug rendering shows actor names for debugging. */
+extern int32 GAnimationBudgetDebugVerboseShowActorNames;
+
+/* In Verbose Mode - Controls whether debug rendering shows non ticking component datas. */
+extern int32 GAnimationBudgetDebugVerboseShowNonTicking;
+
+/* In Verbose Mode - Controls whether debug rendering should be shown as 2D on canvas or not. */
+extern int32 GAnimationBudgetDebugVerboseOnCanvas;
+
 #endif
 
 #if !(UE_BUILD_SHIPPING || UE_BUILD_TEST)
 
 /** Parameters used to force state for debugging */
 extern int32 GAnimationBudgetDebugForce;
+extern int32 GAnimationBudgetDebugForceExcludeAlwaysTicked;
 extern int32 GAnimationBudgetDebugForceRate;
 extern int32 GAnimationBudgetDebugForceInterpolation;
 extern int32 GAnimationBudgetDebugForceReducedWork;
